Adding missing BUF_SourceCopy flags to the structured buffer tests.

diff --git a/Engine/Plugins/Tests/RHITests/Source/RHITests/Public/RHIBufferTests.h b/Engine/Plugins/Tests/RHITests/Source/RHITests/Public/RHIBufferTests.h
--- a/Engine/Plugins/Tests/RHITests/Source/RHITests/Public/RHIBufferTests.h
+++ b/Engine/Plugins/Tests/RHITests/Source/RHITests/Public/RHIBufferTests.h
@@ -113,7 +113,7 @@ class FRHIBufferTests
 		check(NumTestBytes == Stride);
 
 		FRHIResourceCreateInfo Info;
-		FStructuredBufferRHIRef StructuredBuffer = RHICreateStructuredBuffer(Stride, BufferSize, BUF_ShaderResource | BUF_UnorderedAccess | InExtraUsage, Info);
+		FStructuredBufferRHIRef StructuredBuffer = RHICreateStructuredBuffer(Stride, BufferSize, BUF_ShaderResource | BUF_UnorderedAccess | BUF_SourceCopy | InExtraUsage, Info);
 		FUnorderedAccessViewRHIRef UAV = RHICreateUnorderedAccessView(StructuredBuffer, false, false);
 		FShaderResourceViewRHIRef SRV = RHICreateShaderResourceView(StructuredBuffer);
 		bool bResult = RunTest_UAVClear_Buffer(RHICmdList, *FString::Printf(TEXT("RunTest_UAVClear_StructuredBuffer, Stride: %d, Size: %d, ByteAddress: %d"), Stride, BufferSize, (InExtraUsage & BUF_ByteAddressBuffer) ? 1 : 0), StructuredBuffer.GetReference(), UAV, BufferSize, ClearValue, ClearPtr, TestValue);
