Adding support for R32_FLOAT depth buffers.

diff --git a/Engine/Plugins/Tests/RHITests/Source/RHITests/Private/RHITests.spec.cpp b/Engine/Plugins/Tests/RHITests/Source/RHITests/Private/RHITests.spec.cpp
--- a/Engine/Plugins/Tests/RHITests/Source/RHITests/Private/RHITests.spec.cpp
+++ b/Engine/Plugins/Tests/RHITests/Source/RHITests/Private/RHITests.spec.cpp
@@ -49,6 +49,12 @@ void FAutomationExpectedErrorTest::Define()
 			bool bResult = RunOnRenderThreadSynchronous(FRHITextureTests::Test_RHIClearUAV_Texture3D);
 			TestEqual("Clear Texture3D failed", bResult, 1);
 		});
+
+		It("RHI Formats", [this]()
+		{
+			bool bResult = RunOnRenderThreadSynchronous(FRHITextureTests::Test_RHIFormats);
+			TestEqual("RHI Formats failed", bResult, 1);
+		});
 	});
 
 }
\ No newline at end of file
diff --git a/Engine/Plugins/Tests/RHITests/Source/RHITests/Private/RHITestsModule.cpp b/Engine/Plugins/Tests/RHITests/Source/RHITests/Private/RHITestsModule.cpp
--- a/Engine/Plugins/Tests/RHITests/Source/RHITests/Private/RHITestsModule.cpp
+++ b/Engine/Plugins/Tests/RHITests/Source/RHITests/Private/RHITestsModule.cpp
@@ -33,6 +33,10 @@ static bool RunTests_RenderThread(FRHICommandListImmediate& RHICmdList)
 		RUN_TEST(FRHITextureTests::Test_RHIClearUAV_Texture3D(RHICmdList));
 	}
 
+	{
+		RUN_TEST(FRHITextureTests::Test_RHIFormats(RHICmdList));
+	}
+
 	// @todo - add more tests
 	return bResult;
 }
diff --git a/Engine/Plugins/Tests/RHITests/Source/RHITests/Public/RHITextureTests.h b/Engine/Plugins/Tests/RHITests/Source/RHITests/Public/RHITextureTests.h
--- a/Engine/Plugins/Tests/RHITests/Source/RHITests/Public/RHITextureTests.h
+++ b/Engine/Plugins/Tests/RHITests/Source/RHITests/Public/RHITextureTests.h
@@ -230,4 +230,81 @@ public:
 
 		return bResult;
 	}
+
+	static bool Test_RHIFormat_WithParams(FRHICommandListImmediate& RHICmdList, EPixelFormat ResourceFormat, EPixelFormat SRVFormat, EPixelFormat UAVFormat, ETextureCreateFlags Flags)
+	{
+		const uint32 Width = 32;
+		const uint32 Height = 32;
+
+		bool bResult = true;
+		FString TestName = FString::Printf(TEXT("Test_RHIFormat (%s, %s, %s, %d)"), GPixelFormats[ResourceFormat].Name, GPixelFormats[SRVFormat].Name, GPixelFormats[UAVFormat].Name, Flags);
+		{
+			FRHIResourceCreateInfo CreateInfo(*TestName);
+
+			FTextureRHIRef Texture;
+			Texture = RHICreateTexture2D(Width, Height, ResourceFormat, 1, 1, Flags, CreateInfo);
+			bResult = (Texture != nullptr);
+
+			if (Texture && SRVFormat != PF_Unknown)
+			{
+				FRHITextureSRVCreateInfo ViewInfo(0, 1, SRVFormat);
+				FShaderResourceViewRHIRef SRV = RHICreateShaderResourceView(Texture, ViewInfo);
+				bResult = (SRV != nullptr);
+			}
+
+			// TODO
+			if (Texture && UAVFormat != PF_Unknown)
+			{
+				//RHICreateUnorderedAccessView(Texture, 0, UAVFormat);
+			}
+		}
+		RHICmdList.ImmediateFlush(EImmediateFlushType::FlushRHIThreadFlushResources);
+
+		if (bResult)
+		{
+			UE_LOG(LogRHIUnitTestCommandlet, Display, TEXT("Test passed. \"%s\""), *TestName);
+		}
+
+		return bResult;
+	}
+
+	static bool Test_RHIFormat_RenderTargetFormat(FRHICommandListImmediate& RHICmdList, EPixelFormat Format, bool bAllowUAV)
+	{
+		bool bResult = true;
+		RUN_TEST(Test_RHIFormat_WithParams(RHICmdList, Format, PF_Unknown, PF_Unknown, TexCreate_RenderTargetable));
+		RUN_TEST(Test_RHIFormat_WithParams(RHICmdList, Format, Format, PF_Unknown, TexCreate_RenderTargetable | TexCreate_ShaderResource));
+		if (bAllowUAV)
+		{
+			RUN_TEST(Test_RHIFormat_WithParams(RHICmdList, Format, PF_Unknown, Format, TexCreate_RenderTargetable | TexCreate_UAV));
+			RUN_TEST(Test_RHIFormat_WithParams(RHICmdList, Format, Format, Format, TexCreate_RenderTargetable | TexCreate_ShaderResource | TexCreate_UAV));
+		}
+		return bResult;
+	}
+
+	static bool Test_RHIFormat_DepthFormat(FRHICommandListImmediate& RHICmdList, EPixelFormat ResourceFormat)
+	{
+		bool bResult = true;
+		RUN_TEST(Test_RHIFormat_WithParams(RHICmdList, ResourceFormat, PF_Unknown, PF_Unknown, TexCreate_DepthStencilTargetable));
+		RUN_TEST(Test_RHIFormat_WithParams(RHICmdList, ResourceFormat, ResourceFormat, PF_Unknown, TexCreate_DepthStencilTargetable | TexCreate_ShaderResource));
+
+		if (ResourceFormat == PF_DepthStencil)
+		{
+			RUN_TEST(Test_RHIFormat_WithParams(RHICmdList, ResourceFormat, PF_X24_G8, PF_Unknown, TexCreate_DepthStencilTargetable | TexCreate_ShaderResource));
+		}
+
+		return bResult;
+	}
+
+	static bool Test_RHIFormats(FRHICommandListImmediate& RHICmdList)
+	{
+		bool bResult = true;
+		RUN_TEST(Test_RHIFormat_RenderTargetFormat(RHICmdList, PF_R32_FLOAT, true));
+
+		RUN_TEST(Test_RHIFormat_DepthFormat(RHICmdList, PF_DepthStencil));
+		RUN_TEST(Test_RHIFormat_DepthFormat(RHICmdList, PF_ShadowDepth));
+		RUN_TEST(Test_RHIFormat_DepthFormat(RHICmdList, PF_R32_FLOAT));
+		RUN_TEST(Test_RHIFormat_DepthFormat(RHICmdList, PF_D24));
+
+		return bResult;
+	}
 };
diff --git a/Engine/Source/Runtime/D3D12RHI/Private/D3D12RHIPrivate.h b/Engine/Source/Runtime/D3D12RHI/Private/D3D12RHIPrivate.h
--- a/Engine/Source/Runtime/D3D12RHI/Private/D3D12RHIPrivate.h
+++ b/Engine/Source/Runtime/D3D12RHI/Private/D3D12RHIPrivate.h
@@ -1389,6 +1389,17 @@ inline DXGI_FORMAT FindSharedResourceDXGIFormat(DXGI_FORMAT InFormat, bool bSRGB
 	return InFormat;
 }
 
+inline DXGI_FORMAT FindDepthStencilResourceDXGIFormat(DXGI_FORMAT InFormat)
+{
+	switch (InFormat)
+	{
+	case DXGI_FORMAT_R32_FLOAT: return DXGI_FORMAT_R32_TYPELESS;
+	case DXGI_FORMAT_R16_FLOAT: return DXGI_FORMAT_R16_TYPELESS;
+	}
+
+	return InFormat;
+}
+
 inline DXGI_FORMAT GetPlatformTextureResourceFormat(DXGI_FORMAT InFormat, uint32 InFlags)
 {
 	// Find valid shared texture format
@@ -1396,6 +1407,10 @@ inline DXGI_FORMAT GetPlatformTextureResourceFormat(DXGI_FORMAT InFormat, uint32
 	{
 		return FindSharedResourceDXGIFormat(InFormat, InFlags & TexCreate_SRGB);
 	}
+	if (InFlags & TexCreate_DepthStencilTargetable)
+	{
+		return FindDepthStencilResourceDXGIFormat(InFormat);
+	}
 
 	return InFormat;
 }
diff --git a/Engine/Source/Runtime/RHI/Public/RHIUtilities.h b/Engine/Source/Runtime/RHI/Public/RHIUtilities.h
--- a/Engine/Source/Runtime/RHI/Public/RHIUtilities.h
+++ b/Engine/Source/Runtime/RHI/Public/RHIUtilities.h
@@ -23,6 +23,7 @@ static inline bool IsDepthOrStencilFormat(EPixelFormat Format)
 	case PF_DepthStencil:
 	case PF_X24_G8:
 	case PF_ShadowDepth:
+	case PF_R32_FLOAT:
 		return true;
 
 	default:
diff --git a/Engine/Source/Runtime/RenderCore/Private/RenderTargetPool.cpp b/Engine/Source/Runtime/RenderCore/Private/RenderTargetPool.cpp
--- a/Engine/Source/Runtime/RenderCore/Private/RenderTargetPool.cpp
+++ b/Engine/Source/Runtime/RenderCore/Private/RenderTargetPool.cpp
@@ -566,7 +566,7 @@ bool FRenderTargetPool::FindFreeElement(
 	check(InputDesc.NumMips > 0);
 
 	// Make sure if requesting a depth format that the clear value is correct
-	ensure(!IsDepthOrStencilFormat(InputDesc.Format) || (InputDesc.ClearValue.ColorBinding == EClearBinding::ENoneBound || InputDesc.ClearValue.ColorBinding == EClearBinding::EDepthStencilBound));
+	ensure(!(InputDesc.Flags & TexCreate_DepthStencilTargetable) || (InputDesc.ClearValue.ColorBinding == EClearBinding::ENoneBound || InputDesc.ClearValue.ColorBinding == EClearBinding::EDepthStencilBound));
 
 	// TexCreate_FastVRAM should be used on Desc.Flags
 	ensure(!(InputDesc.TargetableFlags & TexCreate_FastVRAM));
diff --git a/Engine/Source/Runtime/RenderCore/Private/ShaderParameterStruct.cpp b/Engine/Source/Runtime/RenderCore/Private/ShaderParameterStruct.cpp
--- a/Engine/Source/Runtime/RenderCore/Private/ShaderParameterStruct.cpp
+++ b/Engine/Source/Runtime/RenderCore/Private/ShaderParameterStruct.cpp
@@ -360,7 +360,7 @@ bool FDepthStencilBinding::Validate() const
 		EPixelFormat PixelFormat = Texture->Desc.Format;
 		const TCHAR* FormatString = GetPixelFormatString(PixelFormat);
 
-		bool bIsDepthFormat = PixelFormat == PF_DepthStencil || PixelFormat == PF_ShadowDepth || PixelFormat == PF_D24;
+		bool bIsDepthFormat = PixelFormat == PF_DepthStencil || PixelFormat == PF_ShadowDepth || PixelFormat == PF_D24 || PixelFormat == PF_R32_FLOAT;
 		checkf(bIsDepthFormat,
 			TEXT("Can't bind texture %s as a depth stencil because its pixel format is %s."),
 			Texture->Name, FormatString);
diff --git a/Engine/Source/Runtime/VulkanRHI/Private/VulkanRHI.cpp b/Engine/Source/Runtime/VulkanRHI/Private/VulkanRHI.cpp
--- a/Engine/Source/Runtime/VulkanRHI/Private/VulkanRHI.cpp
+++ b/Engine/Source/Runtime/VulkanRHI/Private/VulkanRHI.cpp
@@ -1143,25 +1143,25 @@ void FVulkanDynamicRHI::RHISubmitCommandsAndFlushGPU()
 
 FTexture2DRHIRef FVulkanDynamicRHI::RHICreateTexture2DFromResource(EPixelFormat Format, uint32 SizeX, uint32 SizeY, uint32 NumMips, uint32 NumSamples, VkImage Resource, ETextureCreateFlags Flags)
 {
-	const FRHIResourceCreateInfo ResourceCreateInfo(TEXT("FVulkanTexture2D"), IsDepthOrStencilFormat(Format) ? FClearValueBinding::DepthZero : FClearValueBinding::Transparent);
+	const FRHIResourceCreateInfo ResourceCreateInfo(TEXT("FVulkanTexture2D"), (Flags & TexCreate_DepthStencilTargetable) ? FClearValueBinding::DepthZero : FClearValueBinding::Transparent);
 	return new FVulkanTexture2D(*Device, Format, SizeX, SizeY, NumMips, NumSamples, Resource, Flags, ResourceCreateInfo);
 }
 
 FTexture2DRHIRef FVulkanDynamicRHI::RHICreateTexture2DFromResource(EPixelFormat Format, uint32 SizeX, uint32 SizeY, uint32 NumMips, uint32 NumSamples, VkImage Resource, FSamplerYcbcrConversionInitializer& ConversionInitializer, ETextureCreateFlags Flags)
 {
-	const FRHIResourceCreateInfo ResourceCreateInfo(TEXT("FVulkanTexture2D"), IsDepthOrStencilFormat(Format) ? FClearValueBinding::DepthZero : FClearValueBinding::Transparent);
+	const FRHIResourceCreateInfo ResourceCreateInfo(TEXT("FVulkanTexture2D"), (Flags & TexCreate_DepthStencilTargetable) ? FClearValueBinding::DepthZero : FClearValueBinding::Transparent);
 	return new FVulkanTexture2D(*Device, Format, SizeX, SizeY, NumMips, NumSamples, Resource, ConversionInitializer, Flags, ResourceCreateInfo);
 }
 
 FTexture2DArrayRHIRef FVulkanDynamicRHI::RHICreateTexture2DArrayFromResource(EPixelFormat Format, uint32 SizeX, uint32 SizeY, uint32 ArraySize, uint32 NumMips, uint32 NumSamples, VkImage Resource, ETextureCreateFlags Flags)
 {
-	const FRHIResourceCreateInfo ResourceCreateInfo(TEXT("FVulkanTexture2DArray"), IsDepthOrStencilFormat(Format) ? FClearValueBinding::DepthZero : FClearValueBinding::Transparent);
+	const FRHIResourceCreateInfo ResourceCreateInfo(TEXT("FVulkanTexture2DArray"), (Flags & TexCreate_DepthStencilTargetable) ? FClearValueBinding::DepthZero : FClearValueBinding::Transparent);
 	return new FVulkanTexture2DArray(*Device, Format, SizeX, SizeY, ArraySize, NumMips, NumSamples, Resource, Flags, ResourceCreateInfo);
 }
 
 FTextureCubeRHIRef FVulkanDynamicRHI::RHICreateTextureCubeFromResource(EPixelFormat Format, uint32 Size, bool bArray, uint32 ArraySize, uint32 NumMips, VkImage Resource, ETextureCreateFlags Flags)
 {
-	const FRHIResourceCreateInfo ResourceCreateInfo(TEXT("FVulkanTextureCube"), IsDepthOrStencilFormat(Format) ? FClearValueBinding::DepthZero : FClearValueBinding::Transparent);
+	const FRHIResourceCreateInfo ResourceCreateInfo(TEXT("FVulkanTextureCube"), (Flags & TexCreate_DepthStencilTargetable) ? FClearValueBinding::DepthZero : FClearValueBinding::Transparent);
 	return new FVulkanTextureCube(*Device, Format, Size, bArray, ArraySize, NumMips, Resource, Flags, ResourceCreateInfo);
 }
 
diff --git a/Engine/Source/Runtime/Windows/D3D11RHI/Private/D3D11Texture.cpp b/Engine/Source/Runtime/Windows/D3D11RHI/Private/D3D11Texture.cpp
--- a/Engine/Source/Runtime/Windows/D3D11RHI/Private/D3D11Texture.cpp
+++ b/Engine/Source/Runtime/Windows/D3D11RHI/Private/D3D11Texture.cpp
@@ -496,6 +496,18 @@ inline DXGI_FORMAT FindSharedResourceDXGIFormat(DXGI_FORMAT InFormat, bool bSRGB
 	return InFormat;
 }
 
+inline DXGI_FORMAT FindDepthStencilResourceDXGIFormat(DXGI_FORMAT InFormat)
+{
+	switch (InFormat)
+	{
+	case DXGI_FORMAT_R32_FLOAT: return DXGI_FORMAT_R32_TYPELESS;
+	case DXGI_FORMAT_R16_FLOAT: return DXGI_FORMAT_R16_TYPELESS;
+	}
+
+	return InFormat;
+}
+
+
 DXGI_FORMAT FD3D11DynamicRHI::GetPlatformTextureResourceFormat(DXGI_FORMAT InFormat, uint32 InFlags)
 {
 	// Find valid shared texture format
@@ -503,6 +515,10 @@ DXGI_FORMAT FD3D11DynamicRHI::GetPlatformTextureResourceFormat(DXGI_FORMAT InFor
 	{
 		return FindSharedResourceDXGIFormat(InFormat, InFlags & TexCreate_SRGB);
 	}
+	if (InFlags & TexCreate_DepthStencilTargetable)
+	{
+		return FindDepthStencilResourceDXGIFormat(InFormat);
+	}
 	return InFormat;
 }
 
