Adding format checks to the RHI Texture Tests to avoid failing on platforms that don't support certain formats.

diff --git a/Engine/Plugins/Tests/RHITests/Source/RHITests/Public/RHITextureTests.h b/Engine/Plugins/Tests/RHITests/Source/RHITests/Public/RHITextureTests.h
--- a/Engine/Plugins/Tests/RHITests/Source/RHITests/Public/RHITextureTests.h
+++ b/Engine/Plugins/Tests/RHITests/Source/RHITests/Public/RHITextureTests.h
@@ -271,12 +271,19 @@ public:
 	static bool Test_RHIFormat_RenderTargetFormat(FRHICommandListImmediate& RHICmdList, EPixelFormat Format, bool bAllowUAV)
 	{
 		bool bResult = true;
-		RUN_TEST(Test_RHIFormat_WithParams(RHICmdList, Format, PF_Unknown, PF_Unknown, TexCreate_RenderTargetable));
-		RUN_TEST(Test_RHIFormat_WithParams(RHICmdList, Format, Format, PF_Unknown, TexCreate_RenderTargetable | TexCreate_ShaderResource));
-		if (bAllowUAV)
+		if (EnumHasAnyFlags(GPixelFormats[Format].Capabilities, EPixelFormatCapabilities::RenderTarget))
 		{
-			RUN_TEST(Test_RHIFormat_WithParams(RHICmdList, Format, PF_Unknown, Format, TexCreate_RenderTargetable | TexCreate_UAV));
-			RUN_TEST(Test_RHIFormat_WithParams(RHICmdList, Format, Format, Format, TexCreate_RenderTargetable | TexCreate_ShaderResource | TexCreate_UAV));
+			RUN_TEST(Test_RHIFormat_WithParams(RHICmdList, Format, PF_Unknown, PF_Unknown, TexCreate_RenderTargetable));
+			RUN_TEST(Test_RHIFormat_WithParams(RHICmdList, Format, Format, PF_Unknown, TexCreate_RenderTargetable | TexCreate_ShaderResource));
+			if (bAllowUAV)
+			{
+				RUN_TEST(Test_RHIFormat_WithParams(RHICmdList, Format, PF_Unknown, Format, TexCreate_RenderTargetable | TexCreate_UAV));
+				RUN_TEST(Test_RHIFormat_WithParams(RHICmdList, Format, Format, Format, TexCreate_RenderTargetable | TexCreate_ShaderResource | TexCreate_UAV));
+			}
+		}
+		else
+		{
+			UE_LOG(LogRHIUnitTestCommandlet, Display, TEXT("Skipping test for lack of format support. \"Test_RHIFormat_RenderTargetFormat (%s)\""), GPixelFormats[Format].Name);
 		}
 		return bResult;
 	}
@@ -284,14 +291,20 @@ public:
 	static bool Test_RHIFormat_DepthFormat(FRHICommandListImmediate& RHICmdList, EPixelFormat ResourceFormat)
 	{
 		bool bResult = true;
-		RUN_TEST(Test_RHIFormat_WithParams(RHICmdList, ResourceFormat, PF_Unknown, PF_Unknown, TexCreate_DepthStencilTargetable));
-		RUN_TEST(Test_RHIFormat_WithParams(RHICmdList, ResourceFormat, ResourceFormat, PF_Unknown, TexCreate_DepthStencilTargetable | TexCreate_ShaderResource));
+		if (EnumHasAnyFlags(GPixelFormats[ResourceFormat].Capabilities, EPixelFormatCapabilities::DepthStencil))
+		{
+			RUN_TEST(Test_RHIFormat_WithParams(RHICmdList, ResourceFormat, PF_Unknown, PF_Unknown, TexCreate_DepthStencilTargetable));
+			RUN_TEST(Test_RHIFormat_WithParams(RHICmdList, ResourceFormat, ResourceFormat, PF_Unknown, TexCreate_DepthStencilTargetable | TexCreate_ShaderResource));
 
-		if (ResourceFormat == PF_DepthStencil)
+			if (ResourceFormat == PF_DepthStencil)
+			{
+				RUN_TEST(Test_RHIFormat_WithParams(RHICmdList, ResourceFormat, PF_X24_G8, PF_Unknown, TexCreate_DepthStencilTargetable | TexCreate_ShaderResource));
+			}
+		}
+		else
 		{
-			RUN_TEST(Test_RHIFormat_WithParams(RHICmdList, ResourceFormat, PF_X24_G8, PF_Unknown, TexCreate_DepthStencilTargetable | TexCreate_ShaderResource));
+			UE_LOG(LogRHIUnitTestCommandlet, Display, TEXT("Skipping test for lack of format support. \"Test_RHIFormat_DepthFormat (%s)\""), GPixelFormats[ResourceFormat].Name);
 		}
-
 		return bResult;
 	}
 
