RHI unit tests for texture update 2D

diff --git a/Engine/Plugins/Tests/RHITests/Source/RHITests/Private/RHIBufferTests.cpp b/Engine/Plugins/Tests/RHITests/Source/RHITests/Private/RHIBufferTests.cpp
--- a/Engine/Plugins/Tests/RHITests/Source/RHITests/Private/RHIBufferTests.cpp
+++ b/Engine/Plugins/Tests/RHITests/Source/RHITests/Private/RHIBufferTests.cpp
@@ -14,7 +14,7 @@ bool FRHIBufferTests::VerifyBufferContents(const TCHAR* TestName, FRHICommandLis
 
 		// @todo - readback API is inconsistent across RHIs
 		RHICmdList.SubmitCommandsAndFlushGPU();
-		RHICmdList.BlockUntilGPUIdle();
+		RHICmdList.FlushResources();
 
 		void* Memory = RHILockStagingBuffer(StagingBuffer, 0, NumBytes);
 		Result = VerifyCallback(Memory, NumBytes);
diff --git a/Engine/Plugins/Tests/RHITests/Source/RHITests/Private/RHITests.spec.cpp b/Engine/Plugins/Tests/RHITests/Source/RHITests/Private/RHITests.spec.cpp
--- a/Engine/Plugins/Tests/RHITests/Source/RHITests/Private/RHITests.spec.cpp
+++ b/Engine/Plugins/Tests/RHITests/Source/RHITests/Private/RHITests.spec.cpp
@@ -55,6 +55,12 @@ void FAutomationRHITest::Define()
 			bool bResult = RunOnRenderThreadSynchronous(FRHITextureTests::Test_RHIFormats);
 			TestEqual("RHI Formats failed", bResult, 1);
 		});
+
+		It("RHI Update Texture2D", [this]()
+		{
+			bool bResult = RunOnRenderThreadSynchronous(FRHITextureTests::Test_UpdateTexture2D);
+			TestEqual("RHI Update Texture2D failed", bResult, 1);
+		});
 	});
 
 }
\ No newline at end of file
diff --git a/Engine/Plugins/Tests/RHITests/Source/RHITests/Private/RHITestsModule.cpp b/Engine/Plugins/Tests/RHITests/Source/RHITests/Private/RHITestsModule.cpp
--- a/Engine/Plugins/Tests/RHITests/Source/RHITests/Private/RHITestsModule.cpp
+++ b/Engine/Plugins/Tests/RHITests/Source/RHITests/Private/RHITestsModule.cpp
@@ -31,6 +31,7 @@ static bool RunTests_RenderThread(FRHICommandListImmediate& RHICmdList)
 	{
 		RUN_TEST(FRHITextureTests::Test_RHIClearUAV_Texture2D(RHICmdList));
 		RUN_TEST(FRHITextureTests::Test_RHIClearUAV_Texture3D(RHICmdList));
+		RUN_TEST(FRHITextureTests::Test_UpdateTexture2D(RHICmdList));
 	}
 
 	{
diff --git a/Engine/Plugins/Tests/RHITests/Source/RHITests/Private/RHITextureTests.cpp b/Engine/Plugins/Tests/RHITests/Source/RHITests/Private/RHITextureTests.cpp
--- a/Engine/Plugins/Tests/RHITests/Source/RHITests/Private/RHITextureTests.cpp
+++ b/Engine/Plugins/Tests/RHITests/Source/RHITests/Private/RHITextureTests.cpp
@@ -43,7 +43,7 @@ bool FRHITextureTests::VerifyTextureContents(const TCHAR* TestName, FRHICommandL
 					RHICmdList.WriteGPUFence(GPUFence);
 
 					RHICmdList.SubmitCommandsAndFlushGPU(); // @todo - refactor RHI readback API. This shouldn't be necessary
-					RHICmdList.BlockUntilGPUIdle();
+					RHICmdList.FlushResources();
 
 					int32 Width, Height;
 					void* Ptr;
diff --git a/Engine/Plugins/Tests/RHITests/Source/RHITests/Public/RHITextureTests.h b/Engine/Plugins/Tests/RHITests/Source/RHITests/Public/RHITextureTests.h
--- a/Engine/Plugins/Tests/RHITests/Source/RHITests/Public/RHITextureTests.h
+++ b/Engine/Plugins/Tests/RHITests/Source/RHITests/Public/RHITextureTests.h
@@ -3,6 +3,8 @@
 #pragma once
 
 #include "RHITestsCommon.h"
+#include <type_traits>
+#include "Math/PackedVector.h"
 
 class FRHITextureTests
 {
@@ -320,4 +322,404 @@ public:
 
 		return bResult;
 	}
+
+	template <typename T, std::enable_if_t<std::is_same<T,float>::value || std::is_same<T,FFloat16>::value, bool> = true>
+	static void FillValues(uint32 NumComponents, uint32 Width, uint32 Height, T* Values)
+	{
+		uint32 index = 0;
+		for (uint32 Y = 0; Y < Height; Y++)
+		{
+			for (uint32 X = 0; X < Width; X++)
+			{
+				for (uint32 Component = 0; Component < NumComponents; Component++)
+				{
+					Values[index] = float(index) / (float)(NumComponents * Width * Height);
+					index++;
+				}
+			}
+		}
+	}
+
+	template <typename T, std::enable_if_t<std::is_unsigned<T>::value, bool> = true>
+	static void FillValues(uint32 NumComponents, uint32 Width, uint32 Height, T* Values)
+	{
+		static_assert(std::is_integral<T>::value, "Integral required.");
+		uint64 index = 0;
+		uint64 Value = 0;
+		const T MinInt = T(1) << (sizeof(T) * 8 - 1);
+		for (uint32 Y = 0; Y < Height; Y++)
+		{
+			for (uint32 X = 0; X < Width; X++)
+			{
+				for (uint32 Component = 0; Component < NumComponents; Component++)
+				{
+					// SNORM has two values for -1.0 (for example 0x80 and 0x81) so that 0.0 exists in the encoding
+					// Skip the lowest signed value as texture copy/update might write different encoding
+					if (T(Value) == MinInt)
+					{
+						Value++;
+					}
+					Values[index++] = T(Value++);
+				}
+			}
+		}
+	}
+
+	template <typename T, std::enable_if_t<std::is_same<T,FFloat3Packed>::value, bool> = true>
+	static void FillValues(uint32 NumComponents, uint32 Width, uint32 Height, T* Values)
+	{
+		uint32 ValueIndex = 0;
+		uint32 ComponentIndex = 0;
+		for (uint32 Y = 0; Y < Height; Y++)
+		{
+			for (uint32 X = 0; X < Width; X++)
+			{
+				FLinearColor Color;
+				Color.R = float(ComponentIndex++ % (1 << 11)) / (float)(1 << 11);
+				Color.G = float(ComponentIndex++ % (1 << 11)) / (float)(1 << 11);
+				Color.B = float(ComponentIndex++ % (1 << 11)) / (float)(1 << 11);
+				Color.A = 0.0f;
+				Values[ValueIndex] = T(Color);
+				ValueIndex++;
+			}
+		}
+	}
+
+	template <typename T, uint32 NumComponents, bool bInFloatUAV>
+	struct FDataSource
+	{
+		static const bool bFloatUAV = bInFloatUAV;
+		static const uint32 ElementSize = sizeof(T) * NumComponents;
+
+		static void FillSourceData(uint32 Width, uint32 Height, uint8* OutSourceData)
+		{
+			FillValues(NumComponents, Width, Height, (T*)OutSourceData);
+		}
+	};
+
+	typedef FDataSource< uint8,			1, true >	FDataSource8x1;
+	typedef FDataSource< uint8,			2, true >	FDataSource8x2;
+	typedef FDataSource< uint8,			4, true >	FDataSource8x4;
+	typedef FDataSource< uint16,		1, true >	FDataSource16x1;
+	typedef FDataSource< uint16,		2, true >	FDataSource16x2;
+	typedef FDataSource< uint16,		4, true >	FDataSource16x4;
+	typedef FDataSource< uint32,		1, true >	FDataSource32x1;
+	typedef FDataSource< uint32,		2, true >	FDataSource32x2;
+	typedef FDataSource< uint32,		4, true >	FDataSource32x4;
+	typedef FDataSource< FFloat16,		1, true >	FDataSource16x1F;
+	typedef FDataSource< FFloat16,		2, true >	FDataSource16x2F;
+	typedef FDataSource< FFloat16,		4, true >	FDataSource16x4F;
+	typedef FDataSource< float,			1, true >	FDataSource32x1F;
+	typedef FDataSource< float,			2, true >	FDataSource32x2F;
+	typedef FDataSource< float,			3, true >	FDataSource32x3F;
+	typedef FDataSource< float,			4, true >	FDataSource32x4F;
+	typedef FDataSource< uint8,			1, false >	FDataSource8x1UInt;
+	typedef FDataSource< uint8,			2, false >	FDataSource8x2UInt;
+	typedef FDataSource< uint8,			4, false >	FDataSource8x4UInt;
+	typedef FDataSource< uint16,		1, false >	FDataSource16x1UInt;
+	typedef FDataSource< uint16,		2, false >	FDataSource16x2UInt;
+	typedef FDataSource< uint16,		4, false >	FDataSource16x4UInt;
+	typedef FDataSource< uint32,		1, false >	FDataSource32x1UInt;
+	typedef FDataSource< uint32,		2, false >	FDataSource32x2UInt;
+	typedef FDataSource< uint32,		3, false >	FDataSource32x3UInt;
+	typedef FDataSource< uint32,		4, false >	FDataSource32x4UInt;
+	typedef FDataSource< uint64,		1, false >	FDataSource64x1UInt;
+	typedef FDataSource< FFloat3Packed,	1, true >	FDataSource11_11_10F;
+
+	template <typename SourceType>
+	static bool Test_UpdateTexture2D_Impl(FRHICommandListImmediate& RHICmdList, FString TestName, FRHITexture2D* Texture, const FUpdateTextureRegion2D& Region, uint32 SourcePitch, const uint8* SourceData, const uint8* ZeroData)
+	{
+		bool bResult = true;
+		
+		auto VerifyMip = [&](void* Ptr, uint32 MipWidth, uint32 MipHeight, uint32 Width, uint32 Height, uint32 CurrentMipIndex, uint32 CurrentSliceIndex, bool bAlwaysZero)
+		{
+			uint32 BytesPerPixel = GPixelFormats[Texture->GetFormat()].BlockBytes;
+			check(CurrentMipIndex == 0);
+			check(BytesPerPixel == SourceType::ElementSize);
+
+			for (uint32 Y = 0; Y < MipHeight; ++Y)
+			{
+				uint8* Row = ((uint8*)Ptr) + (Y * Width * BytesPerPixel);
+
+				// Verify row within mip stride bounds matches the expected value
+				for (uint32 X = 0; X < MipWidth; ++X)
+				{
+					uint8* Pixel = Row + X * BytesPerPixel;
+
+					bool bShouldBeZero = bAlwaysZero || X < Region.DestX || Y < Region.DestY || X >= (Region.DestX + Region.Width) || Y >= (Region.DestY + Region.Height);
+					if (bShouldBeZero)
+					{
+						if (!IsZeroMem(Pixel, BytesPerPixel))
+							return false;
+					}
+					else
+					{
+						const uint8* TestPixel = SourceData + (Region.SrcX + X - Region.DestX) * BytesPerPixel + (Region.SrcY + Y - Region.DestY) * SourcePitch;
+						if (FMemory::Memcmp(Pixel, TestPixel, BytesPerPixel) != 0)
+							return false;
+					}
+				}
+			}
+
+			return true;
+		};
+
+		// clear to zero
+		RHICmdList.Transition(FRHITransitionInfo(Texture, ERHIAccess::Unknown, ERHIAccess::UAVCompute));
+		FUpdateTextureRegion2D ZeroRegion = { 0, 0, 0, 0, Texture->GetSizeX(), Texture->GetSizeY() };
+		RHICmdList.UpdateTexture2D(Texture, 0, ZeroRegion, SourcePitch, ZeroData);
+		RHICmdList.Transition(FRHITransitionInfo(Texture, ERHIAccess::UAVCompute, ERHIAccess::CopySrc));
+		bResult &= VerifyTextureContents(*FString::Printf(TEXT("%s - clear whole resource to zero"), *TestName), RHICmdList, Texture,
+		[&](void* Ptr, uint32 MipWidth, uint32 MipHeight, uint32 Width, uint32 Height, uint32 CurrentMipIndex, uint32 CurrentSliceIndex)
+		{
+			return VerifyMip(Ptr, MipWidth, MipHeight, Width, Height, CurrentMipIndex, CurrentSliceIndex, true);
+		});
+		RHICmdList.Transition(FRHITransitionInfo(Texture, ERHIAccess::CopySrc, ERHIAccess::UAVCompute));
+
+		// update the texture
+		RHICmdList.UpdateTexture2D(Texture, 0, Region, SourcePitch, SourceData);
+		RHICmdList.Transition(FRHITransitionInfo(Texture, ERHIAccess::UAVCompute, ERHIAccess::CopySrc));
+
+		bResult &= VerifyTextureContents(*FString::Printf(TEXT("%s - update (%d,%d -> %d,%d)"), *TestName, Region.SrcX, Region.SrcY, Region.DestX, Region.DestY), RHICmdList, Texture,
+		[&](void* Ptr, uint32 MipWidth, uint32 MipHeight, uint32 Width, uint32 Height, uint32 CurrentMipIndex, uint32 CurrentSliceIndex)
+		{
+			return VerifyMip(Ptr, MipWidth, MipHeight, Width, Height, CurrentMipIndex, CurrentSliceIndex, false);
+		});
+
+		RHICmdList.ImmediateFlush(EImmediateFlushType::FlushRHIThreadFlushResources);
+
+		return bResult;
+	}
+
+	template <typename SourceType>
+	static bool Test_UpdateTexture2D_Impl(FRHICommandListImmediate& RHICmdList, EPixelFormat Format)
+	{
+		const uint32 TextureWidth = 128;
+		const uint32 TextureHeight = 128;
+		const uint32 SrcDataWidth = 128;
+		const uint32 SrcDataHeight = 128;
+		const uint32 UpdateWidth = 64;
+		const uint32 UpdateHeight = 64;
+		const uint32 UpdateDataPitch = SrcDataWidth * SourceType::ElementSize;
+		uint8 UpdateData[SrcDataWidth * SrcDataHeight * SourceType::ElementSize];
+		uint8 ZeroData[SrcDataWidth * SrcDataHeight * SourceType::ElementSize] = {};
+		SourceType::FillSourceData(SrcDataWidth, SrcDataHeight, UpdateData);
+
+		bool bResult = true;
+		FString TestName = FString::Printf(TEXT("Test_UpdateTexture2D (%s)"), GPixelFormats[Format].Name);
+		{
+			if (!GPixelFormats[Format].Supported)
+			{
+				UE_LOG(LogRHIUnitTestCommandlet, Display, TEXT("Test skipped (format not supported). \"%s\""), *TestName);
+				return true;
+			}
+			if (!RHIPixelFormatHasCapabilities(Format, EPixelFormatCapabilities::Texture2D))
+			{
+				UE_LOG(LogRHIUnitTestCommandlet, Display, TEXT("Test skipped (format not supported as Texture2D). \"%s\""), *TestName);
+				return true;
+			}
+
+			FRHIResourceCreateInfo CreateInfo(*TestName);
+
+			FTexture2DRHIRef Texture;
+			ETextureCreateFlags TexCreateFlags = TexCreate_ShaderResource | TexCreate_UAV;
+			Texture = RHICreateTexture2D(TextureWidth, TextureHeight, Format, 1, 1, TexCreateFlags, CreateInfo);
+			if (Texture == nullptr)
+			{
+				UE_LOG(LogRHIUnitTestCommandlet, Display, TEXT("Test failed (couldn't create texture). \"%s\""), *TestName);
+				return false;
+			}
+
+			auto AxisSlotToCoord = [](uint32 Slot, uint32 TextureSize, uint32 UpdateSize) -> uint32
+			{
+				switch (Slot)
+				{
+				default: checkNoEntry();
+				case 0: return 0;
+				case 1: return (TextureSize / 2) - (UpdateSize / 2);
+				case 2: return TextureSize - UpdateSize;
+				}
+			};
+
+			FUpdateTextureRegion2D Region(0, 0, 0, 0, UpdateWidth, UpdateHeight);
+			// Test destination offsets
+			for (uint32 DestRow = 0; DestRow < 3; DestRow++)
+			{
+				Region.DestY = AxisSlotToCoord(DestRow, TextureHeight, UpdateHeight);
+				for (uint32 DestColumn = 0; DestColumn < 3; DestColumn++)
+				{
+					Region.DestX = AxisSlotToCoord(DestColumn, TextureWidth, UpdateWidth);
+					bResult &= Test_UpdateTexture2D_Impl<SourceType>(RHICmdList, TestName, Texture, Region, UpdateDataPitch, UpdateData, ZeroData);
+				}
+			}
+
+#if 0 // @todo enable when we want to support source offsets
+			// Test source offsets
+			for (uint32 SrcRow = 0; SrcRow < 3; SrcRow++)
+			{
+				Region.SrcY = AxisSlotToCoord(SrcRow, SrcDataHeight, UpdateHeight);
+				for (uint32 SrcColumn = 0; SrcColumn < 3; SrcColumn++)
+				{
+					Region.SrcX = AxisSlotToCoord(SrcColumn, SrcDataWidth, UpdateWidth);
+					bResult &= Test_UpdateTexture2D_Impl<SourceType>(RHICmdList, TestName, Texture, Region, UpdateDataPitch, UpdateData, ZeroData);
+				}
+			}
+#endif
+		}
+
+		return bResult;
+	}
+
+	static bool Test_UpdateTexture2D(FRHICommandListImmediate& RHICmdList)
+	{
+		bool bResult = true;
+		auto RunTestOnFormat = [&](EPixelFormat Format)
+		{
+			switch (Format)
+			{
+			default:
+				checkNoEntry();
+			case PF_Unknown:
+			case PF_DXT1:
+			case PF_DXT3:
+			case PF_DXT5:
+			case PF_UYVY:
+			case PF_DepthStencil:
+			case PF_ShadowDepth:
+			case PF_D24:
+			case PF_BC5:
+			case PF_A1:
+			case PF_PVRTC2:
+			case PF_PVRTC4:
+			case PF_BC4:
+			case PF_ATC_RGB:
+			case PF_ATC_RGBA_E:
+			case PF_ATC_RGBA_I:
+			case PF_X24_G8:
+			case PF_ETC1:
+			case PF_ETC2_RGB:
+			case PF_ETC2_RGBA:
+			case PF_ASTC_4x4:
+			case PF_ASTC_6x6:
+			case PF_ASTC_8x8:
+			case PF_ASTC_10x10:
+			case PF_ASTC_12x12:
+			case PF_BC6H:
+			case PF_BC7:
+			case PF_XGXR8:
+			case PF_PLATFORM_HDR_0:
+			case PF_PLATFORM_HDR_1:
+			case PF_PLATFORM_HDR_2:
+			case PF_NV12:
+			case PF_ETC2_R11_EAC:
+			case PF_ETC2_RG11_EAC:
+			case PF_ASTC_4x4_HDR:
+			case PF_ASTC_6x6_HDR:
+			case PF_ASTC_8x8_HDR:
+			case PF_ASTC_10x10_HDR:
+			case PF_ASTC_12x12_HDR:
+			case PF_R32G32B32F:
+			case PF_R32G32B32_UINT:
+			case PF_R32G32B32_SINT:
+				break;
+
+			case PF_G8:
+			case PF_A8:
+			case PF_L8:
+			case PF_R8:
+				RUN_TEST(Test_UpdateTexture2D_Impl<FDataSource8x1>(RHICmdList, Format));
+				break;
+			case PF_R8_UINT:
+			case PF_R8_SINT:
+				RUN_TEST(Test_UpdateTexture2D_Impl<FDataSource8x1UInt>(RHICmdList, Format));
+				break;
+			case PF_V8U8:
+			case PF_R8G8:
+				RUN_TEST(Test_UpdateTexture2D_Impl<FDataSource8x2>(RHICmdList, Format));
+				break;
+			case PF_R8G8_UINT:
+				RUN_TEST(Test_UpdateTexture2D_Impl<FDataSource8x2UInt>(RHICmdList, Format));
+				break;
+			case PF_B8G8R8A8:
+			case PF_R8G8B8A8:
+			case PF_A8R8G8B8:
+			case PF_R8G8B8A8_SNORM:
+				RUN_TEST(Test_UpdateTexture2D_Impl<FDataSource8x4>(RHICmdList, Format));
+				break;
+			case PF_R8G8B8A8_UINT:
+				RUN_TEST(Test_UpdateTexture2D_Impl<FDataSource8x4UInt>(RHICmdList, Format));
+				break;
+			case PF_G16:
+			case PF_R5G6B5_UNORM:	// only for the 16-bit pixel data
+			case PF_B5G5R5A1_UNORM:	// only for the 16-bit pixel data
+				RUN_TEST(Test_UpdateTexture2D_Impl<FDataSource16x1>(RHICmdList, Format));
+				break;
+			case PF_R16_UINT:
+			case PF_R16_SINT:
+				RUN_TEST(Test_UpdateTexture2D_Impl<FDataSource16x1UInt>(RHICmdList, Format));
+				break;
+			case PF_R16F:
+			case PF_R16F_FILTER:
+				RUN_TEST(Test_UpdateTexture2D_Impl<FDataSource16x1F>(RHICmdList, Format));
+				break;
+			case PF_G16R16:
+			case PF_G16R16_SNORM:
+				RUN_TEST(Test_UpdateTexture2D_Impl<FDataSource16x2>(RHICmdList, Format));
+				break;
+			case PF_R16G16_UINT:
+				RUN_TEST(Test_UpdateTexture2D_Impl<FDataSource16x2UInt>(RHICmdList, Format));
+				break;
+			case PF_G16R16F:
+			case PF_G16R16F_FILTER:
+				RUN_TEST(Test_UpdateTexture2D_Impl<FDataSource16x2F>(RHICmdList, Format));
+				break;
+			case PF_A16B16G16R16:
+			case PF_R16G16B16A16_UNORM:
+			case PF_R16G16B16A16_SNORM:
+				RUN_TEST(Test_UpdateTexture2D_Impl<FDataSource16x4>(RHICmdList, Format));
+				break;
+			case PF_R16G16B16A16_SINT:
+			case PF_R16G16B16A16_UINT:
+				RUN_TEST(Test_UpdateTexture2D_Impl<FDataSource16x4UInt>(RHICmdList, Format));
+				break;
+			case PF_R32_SINT:
+			case PF_R32_UINT:
+				RUN_TEST(Test_UpdateTexture2D_Impl<FDataSource32x1UInt>(RHICmdList, Format));
+				break;
+			case PF_R32_FLOAT:
+				RUN_TEST(Test_UpdateTexture2D_Impl<FDataSource32x1F>(RHICmdList, Format));
+				break;
+			case PF_G32R32F:
+				RUN_TEST(Test_UpdateTexture2D_Impl<FDataSource32x2F>(RHICmdList, Format));
+				break;
+			case PF_R32G32_UINT:
+				RUN_TEST(Test_UpdateTexture2D_Impl<FDataSource32x2UInt>(RHICmdList, Format));
+				break;
+			case PF_A32B32G32R32F:
+				RUN_TEST(Test_UpdateTexture2D_Impl<FDataSource32x4F>(RHICmdList, Format));
+				break;
+			case PF_R32G32B32A32_UINT:
+				RUN_TEST(Test_UpdateTexture2D_Impl<FDataSource32x4UInt>(RHICmdList, Format));
+				break;
+			case PF_R64_UINT:
+				RUN_TEST(Test_UpdateTexture2D_Impl<FDataSource64x1UInt>(RHICmdList, Format));
+				break;
+			case PF_FloatRGB:
+			case PF_FloatR11G11B10:
+				RUN_TEST(Test_UpdateTexture2D_Impl<FDataSource11_11_10F>(RHICmdList, Format));
+				break;
+			case PF_A2B10G10R10:
+				RUN_TEST(Test_UpdateTexture2D_Impl<FDataSource8x4>(RHICmdList, Format));
+				break;
+			case PF_FloatRGBA:
+				RUN_TEST(Test_UpdateTexture2D_Impl<FDataSource16x4F>(RHICmdList, Format));
+				break;
+			}
+		};
+#define RUN_TEST_ON_FORMAT(Format) RunTestOnFormat(Format);
+		FOREACH_ENUM_EPIXELFORMAT(RUN_TEST_ON_FORMAT);
+#undef RUN_TEST_ON_FORMAT
+
+		return bResult;
+	}
 };
