Allocate from FMemStack rather than from the thread stack

diff --git a/Engine/Plugins/Tests/RHITests/Source/RHITests/Public/RHITextureTests.h b/Engine/Plugins/Tests/RHITests/Source/RHITests/Public/RHITextureTests.h
--- a/Engine/Plugins/Tests/RHITests/Source/RHITests/Public/RHITextureTests.h
+++ b/Engine/Plugins/Tests/RHITests/Source/RHITests/Public/RHITextureTests.h
@@ -504,9 +504,15 @@ public:
 		const uint32 UpdateWidth = 64;
 		const uint32 UpdateHeight = 64;
 		const uint32 UpdateDataPitch = SrcDataWidth * SourceType::ElementSize;
-		uint8 UpdateData[SrcDataWidth * SrcDataHeight * SourceType::ElementSize];
-		uint8 ZeroData[SrcDataWidth * SrcDataHeight * SourceType::ElementSize] = {};
+		const uint32 SrcDataSize = SrcDataWidth * SrcDataHeight * SourceType::ElementSize;
+		uint8* UpdateData;
+		uint8* ZeroData;
+
+		FMemMark MemMark(FMemStack::Get());
+		UpdateData = new(FMemStack::Get()) uint8[SrcDataSize];
+		ZeroData = new(FMemStack::Get()) uint8[SrcDataSize];
 		SourceType::FillSourceData(SrcDataWidth, SrcDataHeight, UpdateData);
+		FMemory::Memset(ZeroData, 0, SrcDataSize);
 
 		bool bResult = true;
 		FString TestName = FString::Printf(TEXT("Test_UpdateTexture2D (%s)"), GPixelFormats[Format].Name);
