Fix RHI texture unit tests:
* texture format not being passed to the create texture function in the clear tests, typo from CL 19560161.
* update tests trying to create UAV-capable textures without checking for format support; however, they don't actually need a UAV, so I just removed that flag.

#rnx

diff --git a/Engine/Plugins/Tests/RHITests/Source/RHITests/Public/RHITextureTests.h b/Engine/Plugins/Tests/RHITests/Source/RHITests/Public/RHITextureTests.h
--- a/Engine/Plugins/Tests/RHITests/Source/RHITests/Public/RHITextureTests.h
+++ b/Engine/Plugins/Tests/RHITests/Source/RHITests/Public/RHITextureTests.h
@@ -94,6 +94,7 @@ public:
 		{
 			const FRHITextureCreateDesc Desc =
 				FRHITextureCreateDesc(*TestName, (NumSlices == 1) ? ETextureDimension::Texture2D : ETextureDimension::Texture2DArray)
+				.SetFormat(Format)
 				.SetExtent(Width, Height)
 				.SetArraySize(NumSlices)
 				.SetNumMips(NumMips)
@@ -530,7 +531,7 @@ public:
 
 			const FRHITextureCreateDesc Desc =
 				FRHITextureCreateDesc::Create2D(*TestName, TextureWidth, TextureHeight, Format)
-				.SetFlags(ETextureCreateFlags::ShaderResource | ETextureCreateFlags::UAV);
+				.SetFlags(ETextureCreateFlags::ShaderResource);
 
 			FTexture2DRHIRef Texture = RHICreateTexture(Desc);
 			if (Texture == nullptr)
