Fix RHI unit tests: the verification function wasn't waiting for the GPU to finish copying the data, so it was checking uninitialized memory.

Also, the clear value was printed incorrectly after the LWC changes to FVector4.

#rnx

diff --git a/Engine/Plugins/Tests/RHITests/Source/RHITests/Private/RHIBufferTests.cpp b/Engine/Plugins/Tests/RHITests/Source/RHITests/Private/RHIBufferTests.cpp
--- a/Engine/Plugins/Tests/RHITests/Source/RHITests/Private/RHIBufferTests.cpp
+++ b/Engine/Plugins/Tests/RHITests/Source/RHITests/Private/RHIBufferTests.cpp
@@ -13,7 +13,7 @@ bool FRHIBufferTests::VerifyBufferContents(const TCHAR* TestName, FRHICommandLis
 		RHICmdList.CopyToStagingBuffer(Buffer, StagingBuffer, 0, NumBytes);
 
 		// @todo - readback API is inconsistent across RHIs
-		RHICmdList.SubmitCommandsAndFlushGPU();
+		RHICmdList.BlockUntilGPUIdle();
 		RHICmdList.FlushResources();
 
 		void* Memory = RHILockStagingBuffer(StagingBuffer, 0, NumBytes);
diff --git a/Engine/Plugins/Tests/RHITests/Source/RHITests/Public/RHIBufferTests.h b/Engine/Plugins/Tests/RHITests/Source/RHITests/Public/RHIBufferTests.h
--- a/Engine/Plugins/Tests/RHITests/Source/RHITests/Public/RHIBufferTests.h
+++ b/Engine/Plugins/Tests/RHITests/Source/RHITests/Public/RHIBufferTests.h
@@ -31,7 +31,7 @@ class FRHIBufferTests
 		});
 
 		FString ClearValueStr;
-		if (TAreTypesEqual<ValueType, FVector4>::Value)
+		if (TAreTypesEqual<ValueType, FVector4f>::Value)
 		{
 			ClearValueStr = FString::Printf(TEXT("%f %f %f %f"), ClearValue.X, ClearValue.Y, ClearValue.Z, ClearValue.W);
 		}
@@ -54,7 +54,9 @@ class FRHIBufferTests
 			{
 				uint8* Element = ((uint8*)Ptr) + Index * NumTestBytes;
 				if (FMemory::Memcmp(Element, TestValue, NumTestBytes) != 0)
+				{
 					return false;
+				}
 			}
 
 			return true;
