Added RHICopyTexture unit tests and fixed up inconsistencies / bugs.
 - A zero size now uses the source texture size adjusted to the source mip.
 - Source / dest offsets and size now properly scale with each successive mip level.
 - Fixed bug in reflection capture GPU copy due to inconsistent behavior.
 - Deprecated GRHISupportsCopyToTextureMultipleMips which is now always true.

diff --git a/Engine/Plugins/Tests/RHITests/Source/RHITests/Private/RHITests.spec.cpp b/Engine/Plugins/Tests/RHITests/Source/RHITests/Private/RHITests.spec.cpp
--- a/Engine/Plugins/Tests/RHITests/Source/RHITests/Private/RHITests.spec.cpp
+++ b/Engine/Plugins/Tests/RHITests/Source/RHITests/Private/RHITests.spec.cpp
@@ -63,4 +63,12 @@ void FAutomationRHITest::Define()
 		});
 	});
 
+	Describe("Test RHI Copy", [this]
+	{
+		It("RHICopyTexture", [this]()
+		{
+			bool bResult = RunOnRenderThreadSynchronous(FRHITextureTests::Test_RHICopyTexture);
+			TestEqual("RHICopyTexture", bResult, 1);
+		});
+	});
 }
\ No newline at end of file
diff --git a/Engine/Plugins/Tests/RHITests/Source/RHITests/Private/RHITestsModule.cpp b/Engine/Plugins/Tests/RHITests/Source/RHITests/Private/RHITestsModule.cpp
--- a/Engine/Plugins/Tests/RHITests/Source/RHITests/Private/RHITestsModule.cpp
+++ b/Engine/Plugins/Tests/RHITests/Source/RHITests/Private/RHITestsModule.cpp
@@ -38,6 +38,10 @@ static bool RunTests_RenderThread(FRHICommandListImmediate& RHICmdList)
 		RUN_TEST(FRHITextureTests::Test_RHIFormats(RHICmdList));
 	}
 
+	{
+		RUN_TEST(FRHITextureTests::Test_RHICopyTexture(RHICmdList));
+	}
+
 	// @todo - add more tests
 	return bResult;
 }
diff --git a/Engine/Plugins/Tests/RHITests/Source/RHITests/Private/RHITextureTests.cpp b/Engine/Plugins/Tests/RHITests/Source/RHITests/Private/RHITextureTests.cpp
--- a/Engine/Plugins/Tests/RHITests/Source/RHITests/Private/RHITextureTests.cpp
+++ b/Engine/Plugins/Tests/RHITests/Source/RHITests/Private/RHITextureTests.cpp
@@ -26,7 +26,7 @@ bool FRHITextureTests::VerifyTextureContents(const TCHAR* TestName, FRHICommandL
 					FTextureRHIRef StagingTexture = RHICreateTexture(Desc);
 
 					FRHICopyTextureInfo CopyInfo = {};
-					CopyInfo.Size = FIntVector(MipWidth, MipHeight, 1); // @todo - required for D3D11RHI to prevent crash in FD3D11DynamicRHI::RHICopyTexture
+					CopyInfo.Size = FIntVector(MipWidth, MipHeight, 1);
 					CopyInfo.SourceMipIndex = MipIndex;
 					if (Texture->GetTexture3D())
 					{
@@ -39,9 +39,10 @@ bool FRHITextureTests::VerifyTextureContents(const TCHAR* TestName, FRHICommandL
 					}
 					CopyInfo.NumSlices = 1;
 					CopyInfo.NumMips = 1;
+
 					RHICmdList.CopyTexture(Texture, StagingTexture, CopyInfo);
 
-					RHICmdList.Transition(FRHITransitionInfo(StagingTexture, ERHIAccess::CopyDest, ERHIAccess::CopySrc));
+					RHICmdList.Transition(FRHITransitionInfo(StagingTexture, ERHIAccess::CopyDest, ERHIAccess::CPURead));
 
 					FGPUFenceRHIRef GPUFence = RHICreateGPUFence(TEXT("ReadbackFence"));
 					RHICmdList.WriteGPUFence(GPUFence);
diff --git a/Engine/Plugins/Tests/RHITests/Source/RHITests/Public/RHITextureTests.h b/Engine/Plugins/Tests/RHITests/Source/RHITests/Public/RHITextureTests.h
--- a/Engine/Plugins/Tests/RHITests/Source/RHITests/Public/RHITextureTests.h
+++ b/Engine/Plugins/Tests/RHITests/Source/RHITests/Public/RHITextureTests.h
@@ -6,6 +6,8 @@
 #include <type_traits>
 #include "Math/PackedVector.h"
 
+PRAGMA_DISABLE_OPTIMIZATION
+
 class FRHITextureTests
 {
 	static bool VerifyTextureContents(const TCHAR* TestName, FRHICommandListImmediate& RHICmdList, FRHITexture* Texture, TFunctionRef<bool(void* Ptr, uint32 MipWidth, uint32 MipHeight, uint32 Width, uint32 Height, uint32 MipIndex, uint32 SliceIndex)> VerifyCallback);
@@ -732,4 +734,382 @@ public:
 
 		return bResult;
 	}
+
+	//////////////////////////////////////////////////////////////////////////
+
+	static bool Test_RHICopyTexture(FRHICommandListImmediate& RHICmdList)
+	{
+		bool bResult = true;
+
+		const auto WriteTestData = [](void* Ptr, int32 Stride, FIntPoint Size, FIntPoint Offset, int32 ImageIndex)
+		{
+			int32 Index = ImageIndex;
+			uint8* Bytes = (uint8*)Ptr;
+
+			for (int32 Y = Offset.Y; Y < Offset.Y + Size.Y; Y++)
+			{
+				uint32* Row = (uint32*)(Bytes + Y * Stride);
+
+				for (int32 X = Offset.X; X < Offset.X + Size.X; X++)
+				{
+					Row[X] = Index;
+					Index++;
+				}
+			}
+		};
+
+		const auto DumpTestData = [](const void* Ptr, int32 Stride, FIntPoint Size, FIntPoint Offset = FIntPoint::ZeroValue)
+		{
+			FString DataString;
+			uint8* Bytes = (uint8*)Ptr;
+
+			for (int32 Y = Offset.Y; Y < Offset.Y + Size.Y; Y++)
+			{
+				const uint32* Row = (const uint32*)(Bytes + Y * Stride);
+
+				for (int32 X = Offset.X; X < Offset.X + Size.X; X++)
+				{
+					DataString += FString::Printf(TEXT(" %d"), Row[X]);
+				}
+
+				DataString += TEXT("\n");
+			}
+
+			return DataString;
+		};
+
+		const auto CheckTestData = [&](const void* Ptr, int32 Stride, FIntPoint Size, FIntPoint  Offset, int32 ImageIndex) -> bool
+		{
+			int32 Index = ImageIndex;
+			const uint8* Bytes = (uint8*)Ptr;
+
+			for (int32 Y = Offset.Y; Y < Offset.Y + Size.Y; Y++)
+			{
+				const uint32* Row = (const uint32*)(Bytes + Y * Stride);
+
+				for (int32 X = Offset.X; X < Offset.X + Size.X; X++)
+				{
+					if (Row[X] != Index)
+					{
+						UE_LOG(LogRHIUnitTestCommandlet, Error,
+							TEXT("Dumping Test Data:\n")
+							TEXT("\tCoordinates: [%d, %d]\n")
+							TEXT("\tExpected Value: %d\n")
+							TEXT("\tActual Value: %d\n")
+							TEXT("%s"),
+							X, Y, Index, Row[X], *DumpTestData(Ptr, Stride, Size, Offset));
+
+						return false;
+					}
+
+					Index++;
+				}
+			}
+
+			return true;
+		};
+
+		{
+			// Copy to identical texture
+			{
+				const TCHAR* TestName = TEXT("Copy Texture Whole");
+				const FIntPoint Extent(16, 16);
+				const FIntPoint PatternOffset(0, 0);
+				const FIntPoint PatternSize(16, 16);
+				const int32 ImageIndex = 1;
+
+				const FRHITextureCreateDesc SourceDesc =
+					FRHITextureCreateDesc::Create2D(TestName)
+					.SetExtent(Extent.X, Extent.Y)
+					.SetFormat(PF_R32_UINT)
+					.SetInitialState(ERHIAccess::CopySrc);
+
+				FTextureRHIRef SourceTexture = RHICreateTexture(SourceDesc);
+
+				uint32 Stride = 0;
+				void* Data = RHICmdList.LockTexture2D(SourceTexture, 0, RLM_WriteOnly, Stride, false);
+				WriteTestData(Data, Stride, PatternSize, PatternOffset, ImageIndex);
+				RHICmdList.UnlockTexture2D(SourceTexture, 0, false);
+
+				const FRHITextureCreateDesc DestDesc =
+					FRHITextureCreateDesc::Create2D(TestName)
+					.SetExtent(Extent.X, Extent.Y)
+					.SetFormat(PF_R32_UINT)
+					.SetInitialState(ERHIAccess::CopyDest);
+
+				FTextureRHIRef DestTexture = RHICreateTexture(DestDesc);
+				RHICmdList.CopyTexture(SourceTexture, DestTexture, {});
+				RHICmdList.Transition(FRHITransitionInfo(DestTexture, ERHIAccess::CopyDest, ERHIAccess::CopySrc));
+
+				RUN_TEST(VerifyTextureContents(TestName, RHICmdList, DestTexture, [&](void* Data, uint32 MipWidth, uint32 MipHeight, uint32 Width, uint32 Height, uint32 MipIndex, uint32 SliceIndex)
+				{
+					return CheckTestData(Data, Width * sizeof(uint32), PatternSize, PatternOffset, ImageIndex);
+				}));
+			}
+
+			// Copy mip to another texture.
+			const auto CopyTextureMipsLambda = [&](bool bUseExplicitSize)
+			{
+				const FIntPoint Extent(32, 32);
+				const int32 NumMips = 3;
+
+				const FRHITextureCreateDesc SourceDesc =
+					FRHITextureCreateDesc::Create2D(TEXT("Copy Texture Mips"))
+					.SetExtent(Extent.X, Extent.Y)
+					.SetNumMips(NumMips)
+					.SetFormat(PF_R32_UINT)
+					.SetInitialState(ERHIAccess::CopySrc);
+
+				FTextureRHIRef SourceTexture = RHICreateTexture(SourceDesc);
+
+				for (int32 MipIndex = 0; MipIndex < NumMips; ++MipIndex)
+				{
+					uint32 Stride = 0;
+					void* Data = RHICmdList.LockTexture2D(SourceTexture, MipIndex, RLM_WriteOnly, Stride, false);
+					WriteTestData(Data, Stride, FIntPoint(Extent.X >> MipIndex, Extent.Y >> MipIndex), FIntPoint::ZeroValue, MipIndex);
+					RHICmdList.UnlockTexture2D(SourceTexture, MipIndex, false);
+				}
+
+				for (int32 MipIndex = 0; MipIndex < NumMips; ++MipIndex)
+				{
+					const FString TestName = FString::Printf(TEXT("Copy Texture Mips [Mip %d, Size Explicit: %s]"), MipIndex, bUseExplicitSize ? TEXT("Yes") : TEXT("No"));
+
+					const FRHITextureCreateDesc DestDesc =
+						FRHITextureCreateDesc::Create2D(*TestName)
+						.SetExtent(Extent.X >> MipIndex, Extent.Y >> MipIndex)
+						.SetFormat(PF_R32_UINT)
+						.SetInitialState(ERHIAccess::CopySrc);
+
+					FTextureRHIRef DestTexture = RHICreateTexture(DestDesc);
+
+					FRHICopyTextureInfo CopyInfo;
+					CopyInfo.Size = bUseExplicitSize ? FIntVector(DestDesc.Extent.X, DestDesc.Extent.Y, 1) : FIntVector::ZeroValue;
+					CopyInfo.SourceMipIndex = MipIndex;
+
+					RHICmdList.Transition(FRHITransitionInfo(DestTexture, ERHIAccess::CopySrc, ERHIAccess::CopyDest));
+					RHICmdList.CopyTexture(SourceTexture, DestTexture, CopyInfo);
+					RHICmdList.Transition(FRHITransitionInfo(DestTexture, ERHIAccess::CopyDest, ERHIAccess::CopySrc));
+
+					RUN_TEST(VerifyTextureContents(*TestName, RHICmdList, DestTexture, [&](void* Data, uint32 MipWidth, uint32 MipHeight, uint32 Width, uint32 Height, uint32, uint32 SliceIndex)
+					{
+						return CheckTestData(Data, Width * sizeof(uint32), DestDesc.Extent, FIntPoint::ZeroValue, MipIndex);
+					}));
+				}
+			};
+
+			CopyTextureMipsLambda(true);
+			CopyTextureMipsLambda(false);
+
+			// Copy region
+			{
+				const TCHAR* TestName = TEXT("Copy Texture Region");
+				const FIntPoint Extent(32, 32);
+				const FIntPoint PatternOffsetSource(8, 4);
+				const FIntPoint PatternOffsetDest(8, 4);
+				const FIntPoint PatternSize(16, 24);
+				const int32 ImageIndex = 1;
+
+				const FRHITextureCreateDesc SourceDesc =
+					FRHITextureCreateDesc::Create2D(TestName)
+					.SetExtent(Extent.X, Extent.Y)
+					.SetFormat(PF_R32_UINT)
+					.SetInitialState(ERHIAccess::CopySrc);
+
+				FTextureRHIRef SourceTexture = RHICreateTexture(SourceDesc);
+
+				uint32 Stride = 0;
+				void* Data = RHICmdList.LockTexture2D(SourceTexture, 0, RLM_WriteOnly, Stride, false);
+				WriteTestData(Data, Stride, PatternSize, PatternOffsetSource, ImageIndex);
+				RHICmdList.UnlockTexture2D(SourceTexture, 0, false);
+
+				const FRHITextureCreateDesc DestDesc =
+					FRHITextureCreateDesc::Create2D(TestName)
+					.SetExtent(Extent.X, Extent.Y)
+					.SetFormat(PF_R32_UINT)
+					.SetInitialState(ERHIAccess::CopyDest);
+
+				FRHICopyTextureInfo CopyInfo;
+				CopyInfo.Size           = FIntVector(PatternSize.X, PatternSize.Y, 1);
+				CopyInfo.SourcePosition = FIntVector(PatternOffsetSource.X, PatternOffsetSource.Y, 0);
+				CopyInfo.DestPosition   = FIntVector(PatternOffsetDest.X, PatternOffsetDest.Y, 0);
+
+				FTextureRHIRef DestTexture = RHICreateTexture(DestDesc);
+				RHICmdList.CopyTexture(SourceTexture, DestTexture, CopyInfo);
+				RHICmdList.Transition(FRHITransitionInfo(DestTexture, ERHIAccess::CopyDest, ERHIAccess::CopySrc));
+
+				RUN_TEST(VerifyTextureContents(TestName, RHICmdList, DestTexture, [&](void* Data, uint32 MipWidth, uint32 MipHeight, uint32 Width, uint32 Height, uint32 MipIndex, uint32 SliceIndex)
+				{
+					return CheckTestData(Data, Width * sizeof(uint32), PatternSize, PatternOffsetDest, ImageIndex);
+				}));
+			}
+
+			// Copy mip chains
+			{
+				const TCHAR* TestName = TEXT("Copy Texture Subresources - Mips");
+				const FIntPoint Extent(128, 128);
+				const int32 NumMips = 4;
+				const int32 NumArraySlices = 4;
+				const FIntPoint PatternOffsetMip0(8, 16);
+				const FIntPoint PatternSizeMip0(64, 96);
+
+				const FRHITextureCreateDesc SourceDesc =
+					FRHITextureCreateDesc::Create2DArray(TestName)
+					.SetExtent(Extent.X, Extent.Y)
+					.SetNumMips(NumMips) // 128, 64, 32, 16
+					.SetArraySize(NumArraySlices)
+					.SetFormat(PF_R32_UINT)
+					.SetInitialState(ERHIAccess::CopySrc);
+
+				FTextureRHIRef SourceTexture = RHICreateTexture(SourceDesc);
+
+				int32 ImageIndex = 1;
+
+				/*   SLICES
+				  M  1  5  9 13
+				  I  2  6 10 14
+				  P  3  7 11 15
+				  S  4  8 12 16
+				*/
+
+				uint32 Strides[NumMips] = {};
+
+				for (int32 ArraySlice = 0; ArraySlice < NumArraySlices; ++ArraySlice)
+				{
+					for (int32 MipIndex = 0; MipIndex < NumMips; ++MipIndex)
+					{
+						const FIntPoint PatternSize(PatternSizeMip0.X     >> MipIndex, PatternSizeMip0.Y   >> MipIndex);
+						const FIntPoint PatternOffset(PatternOffsetMip0.X >> MipIndex, PatternOffsetMip0.Y >> MipIndex);
+
+						void* Data = RHICmdList.LockTexture2DArray(SourceTexture, ArraySlice, MipIndex, RLM_WriteOnly, Strides[MipIndex], false);
+						WriteTestData(Data, Strides[MipIndex], PatternSize, PatternOffset, ImageIndex);
+						RHICmdList.UnlockTexture2DArray(SourceTexture, ArraySlice, MipIndex, false);
+
+						ImageIndex++;
+					}
+				}
+
+				const FRHITextureCreateDesc DestDesc =
+					FRHITextureCreateDesc::Create2DArray(TestName)
+					.SetExtent(Extent.X, Extent.Y)
+					.SetNumMips(NumMips)
+					.SetArraySize(NumArraySlices)
+					.SetFormat(PF_R32_UINT)
+					.SetInitialState(ERHIAccess::CopyDest);
+
+				FTextureRHIRef DestTexture = RHICreateTexture(DestDesc);
+
+				/*   SLICES
+				  M  13  9  5  1
+				  I  14 10  6  2
+				  P  15 11  7  3
+				  S  16 12  8  4
+				*/
+				for (int32 SourceSliceIndex = 0; SourceSliceIndex < NumArraySlices; ++SourceSliceIndex)
+				{
+					int32 DestSliceIndex = NumArraySlices - SourceSliceIndex - 1;
+
+					FRHICopyTextureInfo CopyInfo;
+					CopyInfo.SourcePosition   = FIntVector(PatternOffsetMip0.X, PatternOffsetMip0.Y, 0);
+					CopyInfo.DestPosition     = FIntVector(PatternOffsetMip0.X, PatternOffsetMip0.Y, 0);
+					CopyInfo.Size             = FIntVector(PatternSizeMip0.X, PatternSizeMip0.Y, 1);
+					CopyInfo.NumMips          = NumMips;
+					CopyInfo.SourceSliceIndex = SourceSliceIndex;
+					CopyInfo.DestSliceIndex   = DestSliceIndex;
+
+					RHICmdList.CopyTexture(SourceTexture, DestTexture, CopyInfo);
+				}
+
+				RHICmdList.Transition(FRHITransitionInfo(DestTexture, ERHIAccess::CopyDest, ERHIAccess::CopySrc));
+
+				RUN_TEST(VerifyTextureContents(TestName, RHICmdList, DestTexture, [&](void* Data, uint32 MipWidth, uint32 MipHeight, uint32 Width, uint32 Height, uint32 MipIndex, uint32 SliceIndex)
+				{
+					int32 ImageIndex = (NumArraySlices - SliceIndex - 1) * NumMips + MipIndex + 1;
+
+					const FIntPoint PatternSize(PatternSizeMip0.X     >> MipIndex, PatternSizeMip0.Y   >> MipIndex);
+					const FIntPoint PatternOffset(PatternOffsetMip0.X >> MipIndex, PatternOffsetMip0.Y >> MipIndex);
+
+					return CheckTestData(Data, Width * sizeof(uint32), PatternSize, PatternOffset, ImageIndex);
+				}));
+			}
+
+			// Copy array slices
+			{
+				const TCHAR* TestName = TEXT("Copy Texture Subresources - Slices");
+				const FIntPoint Extent(128, 128);
+				const int32 NumMips = 4;
+				const int32 NumArraySlices = 4;
+				const FIntPoint PatternOffsetMip0(8, 16);
+				const FIntPoint PatternSizeMip0(64, 96);
+
+				const FRHITextureCreateDesc SourceDesc =
+					FRHITextureCreateDesc::Create2DArray(TestName)
+					.SetExtent(Extent.X, Extent.Y)
+					.SetNumMips(NumMips) // 128, 64, 32, 16
+					.SetArraySize(NumArraySlices)
+					.SetFormat(PF_R32_UINT)
+					.SetInitialState(ERHIAccess::CopySrc);
+
+				FTextureRHIRef SourceTexture = RHICreateTexture(SourceDesc);
+
+				int32 ImageIndex = 1;
+
+				uint32 Strides[NumMips] = {};
+
+				for (int32 ArraySlice = 0; ArraySlice < NumArraySlices; ++ArraySlice)
+				{
+					for (int32 MipIndex = 0; MipIndex < NumMips; ++MipIndex)
+					{
+						const FIntPoint PatternSize(PatternSizeMip0.X >> MipIndex, PatternSizeMip0.Y >> MipIndex);
+						const FIntPoint PatternOffset(PatternOffsetMip0.X >> MipIndex, PatternOffsetMip0.Y >> MipIndex);
+
+						void* Data = RHICmdList.LockTexture2DArray(SourceTexture, ArraySlice, MipIndex, RLM_WriteOnly, Strides[MipIndex], false);
+						WriteTestData(Data, Strides[MipIndex], PatternSize, PatternOffset, ImageIndex);
+						RHICmdList.UnlockTexture2DArray(SourceTexture, ArraySlice, MipIndex, false);
+
+						ImageIndex++;
+					}
+				}
+
+				const FRHITextureCreateDesc DestDesc =
+					FRHITextureCreateDesc::Create2DArray(TestName)
+					.SetExtent(Extent.X, Extent.Y)
+					.SetNumMips(NumMips)
+					.SetArraySize(NumArraySlices)
+					.SetFormat(PF_R32_UINT)
+					.SetInitialState(ERHIAccess::CopyDest);
+
+				FTextureRHIRef DestTexture = RHICreateTexture(DestDesc);
+
+				for (int32 MipIndex = 0; MipIndex < NumMips; ++MipIndex)
+				{
+					FRHICopyTextureInfo CopyInfo;
+					CopyInfo.SourcePosition   = FIntVector(PatternOffsetMip0.X >> MipIndex, PatternOffsetMip0.Y >> MipIndex, 0);
+					CopyInfo.DestPosition     = FIntVector(PatternOffsetMip0.X >> MipIndex, PatternOffsetMip0.Y >> MipIndex, 0);
+					CopyInfo.Size             = FIntVector(PatternSizeMip0.X   >> MipIndex, PatternSizeMip0.Y   >> MipIndex, 1);
+					CopyInfo.NumMips          = 1;
+					CopyInfo.NumSlices        = NumArraySlices;
+					CopyInfo.SourceMipIndex   = MipIndex;
+					CopyInfo.DestMipIndex     = MipIndex;
+
+					RHICmdList.CopyTexture(SourceTexture, DestTexture, CopyInfo);
+				}
+
+				RHICmdList.Transition(FRHITransitionInfo(DestTexture, ERHIAccess::CopyDest, ERHIAccess::CopySrc));
+
+				RUN_TEST(VerifyTextureContents(TestName, RHICmdList, DestTexture, [&](void* Data, uint32 MipWidth, uint32 MipHeight, uint32 Width, uint32 Height, uint32 MipIndex, uint32 SliceIndex)
+				{
+					int32 ImageIndex = SliceIndex * NumMips + MipIndex + 1;
+
+					const FIntPoint PatternSize(PatternSizeMip0.X >> MipIndex, PatternSizeMip0.Y >> MipIndex);
+					const FIntPoint PatternOffset(PatternOffsetMip0.X >> MipIndex, PatternOffsetMip0.Y >> MipIndex);
+
+					return CheckTestData(Data, Width * sizeof(uint32), PatternSize, PatternOffset, ImageIndex);
+				}));
+			}
+		}
+
+		return bResult;
+	}
 };
+
+PRAGMA_ENABLE_OPTIMIZATION
\ No newline at end of file
diff --git a/Engine/Source/Runtime/Apple/AGXRHI/Private/AGXRHI.cpp b/Engine/Source/Runtime/Apple/AGXRHI/Private/AGXRHI.cpp
--- a/Engine/Source/Runtime/Apple/AGXRHI/Private/AGXRHI.cpp
+++ b/Engine/Source/Runtime/Apple/AGXRHI/Private/AGXRHI.cpp
@@ -434,9 +434,6 @@ FAGXDynamicRHI::FAGXDynamicRHI(ERHIFeatureLevel::Type RequestedFeatureLevel)
 	}
 	
 #endif
-
-	
-	GRHISupportsCopyToTextureMultipleMips = true;
 		
 	if(
 	   #if PLATFORM_MAC
diff --git a/Engine/Source/Runtime/Apple/AGXRHI/Private/AGXTexture.cpp b/Engine/Source/Runtime/Apple/AGXRHI/Private/AGXTexture.cpp
--- a/Engine/Source/Runtime/Apple/AGXRHI/Private/AGXTexture.cpp
+++ b/Engine/Source/Runtime/Apple/AGXRHI/Private/AGXTexture.cpp
@@ -2010,11 +2010,8 @@ void FAGXRHICommandContext::RHICopyTexture(FRHITexture* SourceTextureRHI, FRHITe
 		
 		if (TextureFormatExactMatch || TextureFormatCompatible)
 		{
-			FIntVector Size = (CopyInfo.Size != FIntVector::ZeroValue) ? CopyInfo.Size : FIntVector(MetalSrcTexture->GetDesc().Extent.X, MetalSrcTexture->GetDesc().Extent.Y, MetalSrcTexture->GetDesc().Depth);
-			
-			mtlpp::Origin SourceOrigin(CopyInfo.SourcePosition.X, CopyInfo.SourcePosition.Y, CopyInfo.SourcePosition.Z);
-			mtlpp::Origin DestinationOrigin(CopyInfo.DestPosition.X, CopyInfo.DestPosition.Y, CopyInfo.DestPosition.Z);
-			
+			const FIntVector Size = CopyInfo.Size == FIntVector::ZeroValue ? MetalSrcTexture->GetDesc().GetSize() >> CopyInfo.SourceMipIndex : CopyInfo.Size;
+
 			FAGXTexture SrcTexture;
 
 			if (TextureFormatExactMatch)
@@ -2046,6 +2043,9 @@ void FAGXRHICommandContext::RHICopyTexture(FRHITexture* SourceTextureRHI, FRHITe
 					mtlpp::Size SourceSize(FMath::Max(Size.X >> MipIndex, 1), FMath::Max(Size.Y >> MipIndex, 1), FMath::Max(Size.Z >> MipIndex, 1));
 					mtlpp::Size DestSize = SourceSize;
 
+					mtlpp::Origin SourceOrigin(CopyInfo.SourcePosition.X >> MipIndex, CopyInfo.SourcePosition.Y >> MipIndex, CopyInfo.SourcePosition.Z >> MipIndex);
+					mtlpp::Origin DestinationOrigin(CopyInfo.DestPosition.X >> MipIndex, CopyInfo.DestPosition.Y >> MipIndex, CopyInfo.DestPosition.Z >> MipIndex);
+
 					if (TextureFormatCompatible)
 					{
 						DestSize.width  *= GPixelFormats[MetalDestTexture->GetDesc().Format].BlockSizeX;
diff --git a/Engine/Source/Runtime/Apple/AGXRHI/Private/Mac/MacPlatformAGXConfig.cpp b/Engine/Source/Runtime/Apple/AGXRHI/Private/Mac/MacPlatformAGXConfig.cpp
--- a/Engine/Source/Runtime/Apple/AGXRHI/Private/Mac/MacPlatformAGXConfig.cpp
+++ b/Engine/Source/Runtime/Apple/AGXRHI/Private/Mac/MacPlatformAGXConfig.cpp
@@ -493,7 +493,6 @@ GRHISupportsBackBufferWithCustomDepthStencil           = MAC_PLATFORM_AGXRHI_SUP
 GRHISupportsBaseVertexIndex                            = MAC_PLATFORM_AGXRHI_SUPPORTS_BASE_VERTEX_INDEX;
 GRHISupportsComplexVariableRateShadingCombinerOps      = MAC_PLATFORM_AGXRHI_SUPPORTS_COMPLEX_VARIABLE_RATE_SHADING_COMBINER_OPS;
 GRHISupportsConservativeRasterization                  = MAC_PLATFORM_AGXRHI_SUPPORTS_CONSERVATIVE_RASTERIZATION;
-GRHISupportsCopyToTextureMultipleMips                  = MAC_PLATFORM_AGXRHI_SUPPORTS_COPY_TO_TEXTURE_MULTIPLE_MIPS;
 GRHISupportsDepthUAV                                   = MAC_PLATFORM_AGXRHI_SUPPORTS_DEPTH_UAV;
 GRHISupportsDirectGPUMemoryLock                        = MAC_PLATFORM_AGXRHI_SUPPORTS_DIRECT_GPU_MEMORY_LOCK;
 GRHISupportsDrawIndirect                               = MAC_PLATFORM_AGXRHI_SUPPORTS_DRAW_INDIRECT;
diff --git a/Engine/Source/Runtime/Apple/MetalRHI/Private/MetalRHI.cpp b/Engine/Source/Runtime/Apple/MetalRHI/Private/MetalRHI.cpp
--- a/Engine/Source/Runtime/Apple/MetalRHI/Private/MetalRHI.cpp
+++ b/Engine/Source/Runtime/Apple/MetalRHI/Private/MetalRHI.cpp
@@ -428,9 +428,6 @@ FMetalDynamicRHI::FMetalDynamicRHI(ERHIFeatureLevel::Type RequestedFeatureLevel)
 	}
 	
 #endif
-
-	
-	GRHISupportsCopyToTextureMultipleMips = true;
 		
 	if(
 	   #if PLATFORM_MAC
diff --git a/Engine/Source/Runtime/Apple/MetalRHI/Private/MetalTexture.cpp b/Engine/Source/Runtime/Apple/MetalRHI/Private/MetalTexture.cpp
--- a/Engine/Source/Runtime/Apple/MetalRHI/Private/MetalTexture.cpp
+++ b/Engine/Source/Runtime/Apple/MetalRHI/Private/MetalTexture.cpp
@@ -2068,11 +2068,8 @@ void FMetalRHICommandContext::RHICopyTexture(FRHITexture* SourceTextureRHI, FRHI
 		
 		if (TextureFormatExactMatch || TextureFormatCompatible)
 		{
-			FIntVector Size = (CopyInfo.Size != FIntVector::ZeroValue) ? CopyInfo.Size : FIntVector(MetalSrcTexture->GetDesc().Extent.X, MetalSrcTexture->GetDesc().Extent.Y, MetalSrcTexture->GetDesc().Depth);
-			
-			mtlpp::Origin SourceOrigin(CopyInfo.SourcePosition.X, CopyInfo.SourcePosition.Y, CopyInfo.SourcePosition.Z);
-			mtlpp::Origin DestinationOrigin(CopyInfo.DestPosition.X, CopyInfo.DestPosition.Y, CopyInfo.DestPosition.Z);
-			
+			const FIntVector Size = CopyInfo.Size == FIntVector::ZeroValue ? MetalSrcTexture->GetDesc().GetSize() >> CopyInfo.SourceMipIndex : CopyInfo.Size;
+
 			FMetalTexture SrcTexture;
 
 			if (TextureFormatExactMatch)
@@ -2104,6 +2101,9 @@ void FMetalRHICommandContext::RHICopyTexture(FRHITexture* SourceTextureRHI, FRHI
 					mtlpp::Size SourceSize(FMath::Max(Size.X >> MipIndex, 1), FMath::Max(Size.Y >> MipIndex, 1), FMath::Max(Size.Z >> MipIndex, 1));
 					mtlpp::Size DestSize = SourceSize;
 
+					mtlpp::Origin SourceOrigin(CopyInfo.SourcePosition.X >> MipIndex, CopyInfo.SourcePosition.Y >> MipIndex, CopyInfo.SourcePosition.Z >> MipIndex);
+					mtlpp::Origin DestinationOrigin(CopyInfo.DestPosition.X >> MipIndex, CopyInfo.DestPosition.Y >> MipIndex, CopyInfo.DestPosition.Z >> MipIndex);
+
 					if (TextureFormatCompatible)
 					{
 						DestSize.width  *= GPixelFormats[MetalDestTexture->GetDesc().Format].BlockSizeX;
diff --git a/Engine/Source/Runtime/D3D12RHI/Private/D3D12RHI.cpp b/Engine/Source/Runtime/D3D12RHI/Private/D3D12RHI.cpp
--- a/Engine/Source/Runtime/D3D12RHI/Private/D3D12RHI.cpp
+++ b/Engine/Source/Runtime/D3D12RHI/Private/D3D12RHI.cpp
@@ -197,7 +197,6 @@ FD3D12DynamicRHI::FD3D12DynamicRHI(const TArray<TSharedPtr<FD3D12Adapter>>& Chos
 	GMaxShadowDepthBufferSizeX = GMaxTextureDimensions;
 	GMaxShadowDepthBufferSizeY = GMaxTextureDimensions;
 	GRHISupportsResolveCubemapFaces = true;
-	GRHISupportsCopyToTextureMultipleMips = true;
 	GRHISupportsArrayIndexFromAnyShader = true;
 
 	GRHIMaxDispatchThreadGroupsPerDimension.X = D3D12_CS_DISPATCH_MAX_THREAD_GROUPS_PER_DIMENSION;
diff --git a/Engine/Source/Runtime/D3D12RHI/Private/D3D12Texture.cpp b/Engine/Source/Runtime/D3D12RHI/Private/D3D12Texture.cpp
--- a/Engine/Source/Runtime/D3D12RHI/Private/D3D12Texture.cpp
+++ b/Engine/Source/Runtime/D3D12RHI/Private/D3D12Texture.cpp
@@ -2919,18 +2919,8 @@ void FD3D12CommandContext::RHICopyTexture(FRHITexture* SourceTextureRHI, FRHITex
 
 	if (!bAllPixels || !bAllSubresources || bReadback)
 	{
-		// Interpret zero size as source size
-		const FIntVector CopySize = CopyInfo.Size == FIntVector::ZeroValue ? SourceTextureRHI->GetSizeXYZ() : CopyInfo.Size;
-
-		// Copy sub texture regions
-		const CD3DX12_BOX SourceBoxD3D(
-			CopyInfo.SourcePosition.X,
-			CopyInfo.SourcePosition.Y,
-			CopyInfo.SourcePosition.Z,
-			CopyInfo.SourcePosition.X + CopySize.X,
-			CopyInfo.SourcePosition.Y + CopySize.Y,
-			CopyInfo.SourcePosition.Z + CopySize.Z
-		);
+		const FIntVector SourceSize = SourceDesc.GetSize();
+		const FIntVector CopySize = CopyInfo.Size == FIntVector::ZeroValue ? SourceSize >> CopyInfo.SourceMipIndex : CopyInfo.Size;
 
 		D3D12_TEXTURE_COPY_LOCATION Src;
 		Src.pResource = SourceTexture->GetResource()->GetResource();
@@ -2963,23 +2953,20 @@ void FD3D12CommandContext::RHICopyTexture(FRHITexture* SourceTextureRHI, FRHITex
 				uint32 SourceMipIndex = CopyInfo.SourceMipIndex + MipIndex;
 				uint32 DestMipIndex   = CopyInfo.DestMipIndex   + MipIndex;
 
-				CD3DX12_BOX MipSourceBoxD3D(
-					SourceBoxD3D.left  >> MipIndex,
-					SourceBoxD3D.top   >> MipIndex,
-					SourceBoxD3D.front >> MipIndex,
-					// Align to block size to pad the copy when processing the last surface texels.
-					// This will give inconsistent results otherwise between different RHI.
-					AlignArbitrary<uint32>(FMath::Max<uint32>(SourceBoxD3D.right  >> MipIndex, 1), SourcePixelFormatInfo.BlockSizeX),
-					AlignArbitrary<uint32>(FMath::Max<uint32>(SourceBoxD3D.bottom >> MipIndex, 1), SourcePixelFormatInfo.BlockSizeY),
-					AlignArbitrary<uint32>(FMath::Max<uint32>(SourceBoxD3D.back   >> MipIndex, 1), SourcePixelFormatInfo.BlockSizeZ)
-					);
+				D3D12_BOX SrcBox;
+				SrcBox.left   = CopyInfo.SourcePosition.X >> MipIndex;
+				SrcBox.top    = CopyInfo.SourcePosition.Y >> MipIndex;
+				SrcBox.front  = CopyInfo.SourcePosition.Z >> MipIndex;
+				SrcBox.right  = AlignArbitrary<uint32>(FMath::Max<uint32>((CopyInfo.SourcePosition.X + CopySize.X) >> MipIndex, 1), SourcePixelFormatInfo.BlockSizeX);
+				SrcBox.bottom = AlignArbitrary<uint32>(FMath::Max<uint32>((CopyInfo.SourcePosition.Y + CopySize.Y) >> MipIndex, 1), SourcePixelFormatInfo.BlockSizeY);
+				SrcBox.back   = AlignArbitrary<uint32>(FMath::Max<uint32>((CopyInfo.SourcePosition.Z + CopySize.Z) >> MipIndex, 1), SourcePixelFormatInfo.BlockSizeZ);
 
 				const uint32 DestX = CopyInfo.DestPosition.X >> MipIndex;
 				const uint32 DestY = CopyInfo.DestPosition.Y >> MipIndex;
 				const uint32 DestZ = CopyInfo.DestPosition.Z >> MipIndex;
 
 				// RHICopyTexture is allowed to copy mip regions only if are aligned on the block size to prevent unexpected / inconsistent results.
-				ensure(MipSourceBoxD3D.left % SourcePixelFormatInfo.BlockSizeX == 0 && MipSourceBoxD3D.top % SourcePixelFormatInfo.BlockSizeY == 0 && MipSourceBoxD3D.front % SourcePixelFormatInfo.BlockSizeZ == 0);
+				ensure(SrcBox.left % SourcePixelFormatInfo.BlockSizeX == 0 && SrcBox.top % SourcePixelFormatInfo.BlockSizeY == 0 && SrcBox.front % SourcePixelFormatInfo.BlockSizeZ == 0);
 				ensure(DestX % DestPixelFormatInfo.BlockSizeX == 0 && DestY % DestPixelFormatInfo.BlockSizeY == 0 && DestZ % DestPixelFormatInfo.BlockSizeZ == 0);
 
 				Src.SubresourceIndex = CalcSubresource(SourceMipIndex, SourceSliceIndex, SourceTextureRHI->GetNumMips());
@@ -2991,10 +2978,10 @@ void FD3D12CommandContext::RHICopyTexture(FRHITexture* SourceTextureRHI, FRHITex
 				}
 
 				CommandListHandle->CopyTextureRegion(
-					&Dst, 
+					&Dst,
 					DestX, DestY, DestZ,
 					&Src,
-					&MipSourceBoxD3D
+					&SrcBox
 				);
 			}
 		}
diff --git a/Engine/Source/Runtime/OpenGLDrv/Private/OpenGLTexture.cpp b/Engine/Source/Runtime/OpenGLDrv/Private/OpenGLTexture.cpp
--- a/Engine/Source/Runtime/OpenGLDrv/Private/OpenGLTexture.cpp
+++ b/Engine/Source/Runtime/OpenGLDrv/Private/OpenGLTexture.cpp
@@ -2087,8 +2087,10 @@ void FOpenGLDynamicRHI::RHICopyTexture(FRHITexture* SourceTextureRHI, FRHITextur
 
 	if (CopyInfo.Size == FIntVector::ZeroValue)
 	{
+		const FRHITextureDesc& SourceDesc = SourceTextureRHI->GetDesc();
+
 		// Copy whole texture when zero vector is specified for region size.
-		FIntVector SrcTexSize = SourceTextureRHI->GetSizeXYZ();
+		FIntVector SrcTexSize = SourceDesc.GetSize();
 		Width = FMath::Max(1, SrcTexSize.X >> CopyInfo.SourceMipIndex);
 		Height = FMath::Max(1, SrcTexSize.Y >> CopyInfo.SourceMipIndex);
 		switch (SourceTexture->Target)
@@ -2116,32 +2118,38 @@ void FOpenGLDynamicRHI::RHICopyTexture(FRHITexture* SourceTextureRHI, FRHITextur
 
 	if (FOpenGL::SupportsCopyImage())
 	{
-		GLint SrcZOffset, DestZOffset;
-		switch (SourceTexture->Target)
-		{
-		case GL_TEXTURE_3D:
-		case GL_TEXTURE_CUBE_MAP:
-			// For cube maps, the Z offsets select the starting faces.
-			SrcZOffset = CopyInfo.SourcePosition.Z;
-			DestZOffset = CopyInfo.DestPosition.Z;
-			break;
-		case GL_TEXTURE_1D_ARRAY:
-		case GL_TEXTURE_2D_ARRAY:
-			// For texture arrays, the Z offsets and depth actually refer to the range of slices to copy.
-			SrcZOffset = CopyInfo.SourceSliceIndex;
-			DestZOffset = CopyInfo.DestSliceIndex;
-			Depth = CopyInfo.NumSlices;
-			break;
-		default:
-			SrcZOffset = 0;
-			DestZOffset = 0;
-			break;
-		}
-
 		for (uint32 MipIndex = 0; MipIndex < CopyInfo.NumMips; ++MipIndex)
 		{
-			FOpenGL::CopyImageSubData(SourceTexture->GetResource(), SourceTexture->Target, SrcMip, CopyInfo.SourcePosition.X, CopyInfo.SourcePosition.Y, SrcZOffset,
-				DestTexture->GetResource(), DestTexture->Target, DestMip, CopyInfo.DestPosition.X, CopyInfo.DestPosition.Y, DestZOffset,
+			GLint SrcZOffset, DestZOffset;
+			switch (SourceTexture->Target)
+			{
+			case GL_TEXTURE_3D:
+			case GL_TEXTURE_CUBE_MAP:
+				// For cube maps, the Z offsets select the starting faces.
+				SrcZOffset = CopyInfo.SourcePosition.Z >> MipIndex;
+				DestZOffset = CopyInfo.DestPosition.Z  >> MipIndex;
+				break;
+			case GL_TEXTURE_1D_ARRAY:
+			case GL_TEXTURE_2D_ARRAY:
+				// For texture arrays, the Z offsets and depth actually refer to the range of slices to copy.
+				SrcZOffset = CopyInfo.SourceSliceIndex;
+				DestZOffset = CopyInfo.DestSliceIndex;
+				Depth = CopyInfo.NumSlices;
+				break;
+			default:
+				SrcZOffset = 0;
+				DestZOffset = 0;
+				break;
+			}
+
+			GLint SourceX = CopyInfo.SourcePosition.X >> MipIndex;
+			GLint SourceY = CopyInfo.SourcePosition.Y >> MipIndex;
+
+			GLint DestX = CopyInfo.DestPosition.X >> MipIndex;
+			GLint DestY = CopyInfo.DestPosition.Y >> MipIndex;
+
+			FOpenGL::CopyImageSubData(SourceTexture->GetResource(), SourceTexture->Target, SrcMip, SourceX, SourceY, SrcZOffset,
+				DestTexture->GetResource(), DestTexture->Target, DestMip, DestX, DestY, DestZOffset,
 				Width, Height, Depth);
 
 			++SrcMip;
@@ -2190,25 +2198,25 @@ void FOpenGLDynamicRHI::RHICopyTexture(FRHITexture* SourceTextureRHI, FRHITextur
 			switch (DestTexture->Target)
 			{
 				case GL_TEXTURE_1D:
-					FOpenGL::CopyTexSubImage1D(DestTexture->Target, DestMip, XOffset, X, 0, Width);
+					FOpenGL::CopyTexSubImage1D(DestTexture->Target, DestMip, XOffset >> MipIndex, X >> MipIndex, 0, Width);
 					break;
 				case GL_TEXTURE_1D_ARRAY:
-					FOpenGL::CopyTexSubImage2D(DestTexture->Target, DestMip, XOffset, CopyInfo.DestSliceIndex + SliceIndex, X, 0, Width, 1);
+					FOpenGL::CopyTexSubImage2D(DestTexture->Target, DestMip, XOffset >> MipIndex, CopyInfo.DestSliceIndex + SliceIndex, X >> MipIndex, 0, Width, 1);
 					break;
 				case GL_TEXTURE_2D:
 				case GL_TEXTURE_RECTANGLE:
-					FOpenGL::CopyTexSubImage2D(DestTexture->Target, DestMip, XOffset, YOffset, X, Y, Width, Height);
+					FOpenGL::CopyTexSubImage2D(DestTexture->Target, DestMip, XOffset >> MipIndex, YOffset >> MipIndex, X >> MipIndex, Y >> MipIndex, Width, Height);
 					break;
 				case GL_TEXTURE_2D_ARRAY:
-					FOpenGL::CopyTexSubImage3D(DestTexture->Target, DestMip, XOffset, YOffset, CopyInfo.DestSliceIndex + SliceIndex, X, Y, Width, Height);
+					FOpenGL::CopyTexSubImage3D(DestTexture->Target, DestMip, XOffset >> MipIndex, YOffset >> MipIndex, CopyInfo.DestSliceIndex + SliceIndex, X >> MipIndex, Y >> MipIndex, Width, Height);
 					break;
 				case GL_TEXTURE_3D:
-					FOpenGL::CopyTexSubImage3D(DestTexture->Target, DestMip, XOffset, YOffset, ZOffset, X, Y, Width, Height);
+					FOpenGL::CopyTexSubImage3D(DestTexture->Target, DestMip, XOffset >> MipIndex, YOffset >> MipIndex, ZOffset >> MipIndex, X >> MipIndex, Y >> MipIndex, Width, Height);
 					break;
 				case GL_TEXTURE_CUBE_MAP:
 					for (int32 FaceIndex = FMath::Min((int32)CopyInfo.NumSlices, 6) - 1; FaceIndex >= 0; FaceIndex--)
 					{
-						FOpenGL::CopyTexSubImage2D(GL_TEXTURE_CUBE_MAP_POSITIVE_X + (uint32)FaceIndex, CopyInfo.DestMipIndex, XOffset, YOffset, X, Y, Width, Height);
+						FOpenGL::CopyTexSubImage2D(GL_TEXTURE_CUBE_MAP_POSITIVE_X + (uint32)FaceIndex, CopyInfo.DestMipIndex, XOffset >> MipIndex, YOffset >> MipIndex, X >> MipIndex, Y >> MipIndex, Width, Height);
 					}					
 					break;
 			}
diff --git a/Engine/Source/Runtime/RHI/Private/RHI.cpp b/Engine/Source/Runtime/RHI/Private/RHI.cpp
--- a/Engine/Source/Runtime/RHI/Private/RHI.cpp
+++ b/Engine/Source/Runtime/RHI/Private/RHI.cpp
@@ -1325,7 +1325,6 @@ namespace RHIConfig
 bool GIsRHIInitialized = false;
 int32 GRHIPersistentThreadGroupCount = 0;
 int32 GMaxTextureMipCount = MAX_TEXTURE_MIP_COUNT;
-bool GRHISupportsCopyToTextureMultipleMips = false;
 bool GSupportsQuadBufferStereo = false;
 FString GRHIAdapterName;
 FString GRHIAdapterInternalDriverVersion;
diff --git a/Engine/Source/Runtime/RHI/Public/RHI.h b/Engine/Source/Runtime/RHI/Public/RHI.h
--- a/Engine/Source/Runtime/RHI/Public/RHI.h
+++ b/Engine/Source/Runtime/RHI/Public/RHI.h
@@ -66,6 +66,7 @@ extern RHI_API int32 GRHIPersistentThreadGroupCount;
 extern RHI_API int32 GMaxTextureMipCount;
 
 /** Does the RHI implements CopyToTexture() with FRHICopyTextureInfo::NumMips > 1 */
+UE_DEPRECATED(5.1, "All RHIs now support copying to multiple mips.")
 extern RHI_API bool GRHISupportsCopyToTextureMultipleMips;
 
 /** true if this platform has quad buffer stereo support. */
diff --git a/Engine/Source/Runtime/RHI/Public/RHICommandList.h b/Engine/Source/Runtime/RHI/Public/RHICommandList.h
--- a/Engine/Source/Runtime/RHI/Public/RHICommandList.h
+++ b/Engine/Source/Runtime/RHI/Public/RHICommandList.h
@@ -3555,37 +3555,13 @@ public:
 		check(SourceTextureRHI && DestTextureRHI);
 		check(SourceTextureRHI != DestTextureRHI);
 		check(IsOutsideRenderPass());
-		if (GRHISupportsCopyToTextureMultipleMips)
-		{
-			if (Bypass())
-			{
-				GetContext().RHICopyTexture(SourceTextureRHI, DestTextureRHI, CopyInfo);
-				return;
-			}
-			ALLOC_COMMAND(FRHICommandCopyTexture)(SourceTextureRHI, DestTextureRHI, CopyInfo);
-		}
-		else
-		{
-			FRHICopyTextureInfo PerMipInfo = CopyInfo;
-			PerMipInfo.NumMips = 1;
-			for (uint32 MipIndex = 0; MipIndex < CopyInfo.NumMips; MipIndex++)
-			{
-				if (Bypass())
-				{
-					GetContext().RHICopyTexture(SourceTextureRHI, DestTextureRHI, PerMipInfo);
-				}
-				else
-				{
-					ALLOC_COMMAND(FRHICommandCopyTexture)(SourceTextureRHI, DestTextureRHI, PerMipInfo);
-				}
 
-				++PerMipInfo.SourceMipIndex;
-				++PerMipInfo.DestMipIndex;
-				PerMipInfo.Size.X = FMath::Max(1, PerMipInfo.Size.X / 2);
-				PerMipInfo.Size.Y = FMath::Max(1, PerMipInfo.Size.Y / 2);
-				PerMipInfo.Size.Z = FMath::Max(1, PerMipInfo.Size.Z / 2);
-			}
+		if (Bypass())
+		{
+			GetContext().RHICopyTexture(SourceTextureRHI, DestTextureRHI, CopyInfo);
+			return;
 		}
+		ALLOC_COMMAND(FRHICommandCopyTexture)(SourceTextureRHI, DestTextureRHI, CopyInfo);
 	}
 
 	FORCEINLINE_DEBUGGABLE void ResummarizeHTile(FRHITexture2D* DepthTexture)
diff --git a/Engine/Source/Runtime/Renderer/Private/ReflectionEnvironment.cpp b/Engine/Source/Runtime/Renderer/Private/ReflectionEnvironment.cpp
--- a/Engine/Source/Runtime/Renderer/Private/ReflectionEnvironment.cpp
+++ b/Engine/Source/Runtime/Renderer/Private/ReflectionEnvironment.cpp
@@ -487,9 +487,6 @@ void FReflectionEnvironmentCubemapArray::ResizeCubemapArrayGPU(uint32 InMaxCubem
 			FRHITransitionInfo(ReflectionEnvs->GetRHI(), ERHIAccess::Unknown, ERHIAccess::CopyDest)
 		});
 
-		FRHICopyTextureInfo CopyInfo;
-		CopyInfo.Size = TexRef->GetSizeXYZ();
-
 		// Copy the cubemaps, remapping the elements as necessary
 		for (int32 SourceCubemapIndex = 0; SourceCubemapIndex < OldMaxCubemaps; SourceCubemapIndex++)
 		{
@@ -501,6 +498,7 @@ void FReflectionEnvironmentCubemapArray::ResizeCubemapArrayGPU(uint32 InMaxCubem
 
 				for (int32 Face = 0; Face < CubeFace_MAX; Face++)
 				{
+					FRHICopyTextureInfo CopyInfo;
 					CopyInfo.SourceSliceIndex = SourceCubemapIndex * CubeFace_MAX + Face;
 					CopyInfo.DestSliceIndex   = DestCubemapIndex   * CubeFace_MAX + Face;
 
diff --git a/Engine/Source/Runtime/VulkanRHI/Private/VulkanRHI.cpp b/Engine/Source/Runtime/VulkanRHI/Private/VulkanRHI.cpp
--- a/Engine/Source/Runtime/VulkanRHI/Private/VulkanRHI.cpp
+++ b/Engine/Source/Runtime/VulkanRHI/Private/VulkanRHI.cpp
@@ -821,8 +821,6 @@ void FVulkanDynamicRHI::InitInstance()
 		// Indicate that the RHI needs to use the engine's deferred deletion queue.
 		GRHINeedsExtraDeletionLatency = true;
 
-		GRHISupportsCopyToTextureMultipleMips = true;
-
 		GMaxShadowDepthBufferSizeX =  FPlatformMath::Min<int32>(Props.limits.maxImageDimension2D, GMaxShadowDepthBufferSizeX);
 		GMaxShadowDepthBufferSizeY =  FPlatformMath::Min<int32>(Props.limits.maxImageDimension2D, GMaxShadowDepthBufferSizeY);
 		GMaxTextureDimensions = Props.limits.maxImageDimension2D;
diff --git a/Engine/Source/Runtime/VulkanRHI/Private/VulkanTexture.cpp b/Engine/Source/Runtime/VulkanRHI/Private/VulkanTexture.cpp
--- a/Engine/Source/Runtime/VulkanRHI/Private/VulkanTexture.cpp
+++ b/Engine/Source/Runtime/VulkanRHI/Private/VulkanTexture.cpp
@@ -2322,8 +2322,10 @@ void FVulkanCommandListContext::RHICopyTexture(FRHITexture* SourceTexture, FRHIT
 	VkCommandBuffer CmdBuffer = InCmdBuffer->GetHandle();
 
 	const FPixelFormatInfo& PixelFormatInfo = GPixelFormats[DestTexture->GetDesc().Format];
-	const FIntVector SourceXYZ = Source->GetSizeXYZ();
-	const FIntVector DestXYZ = Source->GetSizeXYZ();
+	const FRHITextureDesc& SourceDesc = SourceTexture->GetDesc();
+	const FRHITextureDesc& DestDesc = DestTexture->GetDesc();
+	const FIntVector SourceXYZ = SourceDesc.GetSize();
+	const FIntVector DestXYZ = DestDesc.GetSize();
 
 	check(!EnumHasAnyFlags(Source->GetDesc().Flags, TexCreate_CPUReadback));
 	if(EnumHasAllFlags(Dest->GetDesc().Flags, TexCreate_CPUReadback))
@@ -2378,17 +2380,17 @@ void FVulkanCommandListContext::RHICopyTexture(FRHITexture* SourceTexture, FRHIT
 		if (CopyInfo.Size == FIntVector::ZeroValue)
 		{
 			// Copy whole texture when zero vector is specified for region size
-			ensure(SourceXYZ.X <= DestXYZ.X && SourceXYZ.Y <= DestXYZ.Y);
-			Region.extent.width = FMath::Max<uint32>(PixelFormatInfo.BlockSizeX, SourceXYZ.X >> CopyInfo.SourceMipIndex);
+			Region.extent.width  = FMath::Max<uint32>(PixelFormatInfo.BlockSizeX, SourceXYZ.X >> CopyInfo.SourceMipIndex);
 			Region.extent.height = FMath::Max<uint32>(PixelFormatInfo.BlockSizeY, SourceXYZ.Y >> CopyInfo.SourceMipIndex);
-			Region.extent.depth = FMath::Max<uint32>(PixelFormatInfo.BlockSizeZ, SourceXYZ.Z >> CopyInfo.SourceMipIndex);
+			Region.extent.depth  = FMath::Max<uint32>(PixelFormatInfo.BlockSizeZ, SourceXYZ.Z >> CopyInfo.SourceMipIndex);
+			ensure(Region.extent.width <= (uint32)DestXYZ.X && Region.extent.height <= (uint32)DestXYZ.Y);
 		}
 		else
 		{
 			ensure(CopyInfo.Size.X > 0 && CopyInfo.Size.X <= DestXYZ.X && CopyInfo.Size.Y > 0 && CopyInfo.Size.Y <= DestXYZ.Y);
-			Region.extent.width = CopyInfo.Size.X;
-			Region.extent.height = CopyInfo.Size.Y;
-			Region.extent.depth = FMath::Max(1, CopyInfo.Size.Z);
+			Region.extent.width  = FMath::Max(PixelFormatInfo.BlockSizeX, CopyInfo.Size.X);
+			Region.extent.height = FMath::Max(PixelFormatInfo.BlockSizeY, CopyInfo.Size.Y);
+			Region.extent.depth  = FMath::Max(PixelFormatInfo.BlockSizeZ, CopyInfo.Size.Z);
 		}
 		Region.srcSubresource.aspectMask = Source->GetFullAspectMask();
 		Region.srcSubresource.baseArrayLayer = CopyInfo.SourceSliceIndex;
@@ -2426,9 +2428,9 @@ void FVulkanCommandListContext::RHICopyTexture(FRHITexture* SourceTexture, FRHIT
 				Region.dstOffset.y /= 2;
 				Region.dstOffset.z /= 2;
 
-				Region.extent.width  = FMath::Max<uint32>(Region.extent.width / 2, 1);
-				Region.extent.height = FMath::Max<uint32>(Region.extent.height / 2, 1);
-				Region.extent.depth  = FMath::Max<uint32>(Region.extent.depth / 2, 1);
+				Region.extent.width  = FMath::Max<uint32>(Region.extent.width  / 2, PixelFormatInfo.BlockSizeX);
+				Region.extent.height = FMath::Max<uint32>(Region.extent.height / 2, PixelFormatInfo.BlockSizeY);
+				Region.extent.depth  = FMath::Max<uint32>(Region.extent.depth  / 2, PixelFormatInfo.BlockSizeZ);
 
 				// RHICopyTexture is allowed to copy mip regions only if are aligned on the block size to prevent unexpected / inconsistent results.
 				ensure(Region.srcOffset.x % PixelFormatInfo.BlockSizeX == 0 && Region.srcOffset.y % PixelFormatInfo.BlockSizeY == 0 && Region.srcOffset.z % PixelFormatInfo.BlockSizeZ == 0);
diff --git a/Engine/Source/Runtime/Windows/D3D11RHI/Private/D3D11Device.cpp b/Engine/Source/Runtime/Windows/D3D11RHI/Private/D3D11Device.cpp
--- a/Engine/Source/Runtime/Windows/D3D11RHI/Private/D3D11Device.cpp
+++ b/Engine/Source/Runtime/Windows/D3D11RHI/Private/D3D11Device.cpp
@@ -211,7 +211,6 @@ FD3D11DynamicRHI::FD3D11DynamicRHI(IDXGIFactory1* InDXGIFactory1, D3D_FEATURE_LE
 	GMaxShadowDepthBufferSizeY = GMaxTextureDimensions;
 	GSupportsTimestampRenderQueries = true;
 	GRHISupportsResolveCubemapFaces = true;
-	GRHISupportsCopyToTextureMultipleMips = true;
 
 	GRHIMaxDispatchThreadGroupsPerDimension.X = D3D11_CS_DISPATCH_MAX_THREAD_GROUPS_PER_DIMENSION;
 	GRHIMaxDispatchThreadGroupsPerDimension.Y = D3D11_CS_DISPATCH_MAX_THREAD_GROUPS_PER_DIMENSION;
diff --git a/Engine/Source/Runtime/Windows/D3D11RHI/Private/D3D11Texture.cpp b/Engine/Source/Runtime/Windows/D3D11RHI/Private/D3D11Texture.cpp
--- a/Engine/Source/Runtime/Windows/D3D11RHI/Private/D3D11Texture.cpp
+++ b/Engine/Source/Runtime/Windows/D3D11RHI/Private/D3D11Texture.cpp
@@ -2087,29 +2087,35 @@ void FD3D11DynamicRHI::RHICopyTexture(FRHITexture* SourceTextureRHI, FRHITexture
 	{
 		const FPixelFormatInfo& PixelFormatInfo = GPixelFormats[SourceTextureRHI->GetFormat()];
 
-		D3D11_BOX SrcBox;
-		SrcBox.left = CopyInfo.SourcePosition.X;
-		SrcBox.top = CopyInfo.SourcePosition.Y;
-		SrcBox.front = CopyInfo.SourcePosition.Z;
+		const FIntVector SourceSize = SourceDesc.GetSize();
+		const FIntVector CopySize = CopyInfo.Size == FIntVector::ZeroValue ? SourceSize >> CopyInfo.SourceMipIndex : CopyInfo.Size;
 
 		for (uint32 SliceIndex = 0; SliceIndex < CopyInfo.NumSlices; ++SliceIndex)
 		{
 			uint32 SourceSliceIndex = CopyInfo.SourceSliceIndex + SliceIndex;
-			uint32 DestSliceIndex = CopyInfo.DestSliceIndex + SliceIndex;
+			uint32 DestSliceIndex   = CopyInfo.DestSliceIndex   + SliceIndex;
 
 			for (uint32 MipIndex = 0; MipIndex < CopyInfo.NumMips; ++MipIndex)
 			{
 				uint32 SourceMipIndex = CopyInfo.SourceMipIndex + MipIndex;
-				uint32 DestMipIndex = CopyInfo.DestMipIndex + MipIndex;
+				uint32 DestMipIndex   = CopyInfo.DestMipIndex   + MipIndex;
 
 				const uint32 SourceSubresource = D3D11CalcSubresource(SourceMipIndex, SourceSliceIndex, SourceTextureRHI->GetNumMips());
 				const uint32 DestSubresource = D3D11CalcSubresource(DestMipIndex, DestSliceIndex, DestTextureRHI->GetNumMips());
 
-				SrcBox.right = CopyInfo.SourcePosition.X  + AlignArbitrary<uint32>(FMath::Max(CopyInfo.Size.X >> MipIndex, 1), PixelFormatInfo.BlockSizeX);
-				SrcBox.bottom = CopyInfo.SourcePosition.Y + AlignArbitrary<uint32>(FMath::Max(CopyInfo.Size.Y >> MipIndex, 1), PixelFormatInfo.BlockSizeY);
-				SrcBox.back = CopyInfo.SourcePosition.Z   + AlignArbitrary<uint32>(FMath::Max(CopyInfo.Size.Z >> MipIndex, 1), PixelFormatInfo.BlockSizeZ);
+				D3D11_BOX SrcBox;
+				SrcBox.left   = CopyInfo.SourcePosition.X >> MipIndex;
+				SrcBox.top    = CopyInfo.SourcePosition.Y >> MipIndex;
+				SrcBox.front  = CopyInfo.SourcePosition.Z >> MipIndex;
+				SrcBox.right  = AlignArbitrary<uint32>(FMath::Max<uint32>((CopyInfo.SourcePosition.X + CopySize.X) >> MipIndex, 1), PixelFormatInfo.BlockSizeX);
+				SrcBox.bottom = AlignArbitrary<uint32>(FMath::Max<uint32>((CopyInfo.SourcePosition.Y + CopySize.Y) >> MipIndex, 1), PixelFormatInfo.BlockSizeY);
+				SrcBox.back   = AlignArbitrary<uint32>(FMath::Max<uint32>((CopyInfo.SourcePosition.Z + CopySize.Z) >> MipIndex, 1), PixelFormatInfo.BlockSizeZ);
 
-				Direct3DDeviceIMContext->CopySubresourceRegion(DestTexture->GetResource(), DestSubresource, CopyInfo.DestPosition.X, CopyInfo.DestPosition.Y, CopyInfo.DestPosition.Z, SourceTexture->GetResource(), SourceSubresource, &SrcBox);
+				const uint32 DestX = CopyInfo.DestPosition.X >> MipIndex;
+				const uint32 DestY = CopyInfo.DestPosition.Y >> MipIndex;
+				const uint32 DestZ = CopyInfo.DestPosition.Z >> MipIndex;
+
+				Direct3DDeviceIMContext->CopySubresourceRegion(DestTexture->GetResource(), DestSubresource, DestX, DestY, DestZ, SourceTexture->GetResource(), SourceSubresource, &SrcBox);
 			}
 		}
 	}
