Move the RHI unit tests into a new root category called Rendering, and include this into the rendering group.

This way, the unit tests run for every platform and RHI currently configured in CIS, and they run with validation too.

#rnx

diff --git a/Engine/Plugins/Tests/RHITests/Source/RHITests/Private/RHITests.spec.cpp b/Engine/Plugins/Tests/RHITests/Source/RHITests/Private/RHITests.spec.cpp
--- a/Engine/Plugins/Tests/RHITests/Source/RHITests/Private/RHITests.spec.cpp
+++ b/Engine/Plugins/Tests/RHITests/Source/RHITests/Private/RHITests.spec.cpp
@@ -8,7 +8,7 @@
 #include "RHITextureTests.h"
 
 
-BEGIN_DEFINE_SPEC(FAutomationRHITest, "System.Automation.RHI", EAutomationTestFlags::EngineFilter | EAutomationTestFlags::ApplicationContextMask)
+BEGIN_DEFINE_SPEC(FAutomationRHITest, "Rendering.RHI", EAutomationTestFlags::EngineFilter | EAutomationTestFlags::ApplicationContextMask | EAutomationTestFlags::NonNullRHI)
 END_DEFINE_SPEC(FAutomationRHITest)
 void FAutomationRHITest::Define()
 {
@@ -49,13 +49,19 @@ void FAutomationRHITest::Define()
 			bool bResult = RunOnRenderThreadSynchronous(FRHITextureTests::Test_RHIClearUAV_Texture3D);
 			TestEqual("Clear Texture3D failed", bResult, 1);
 		});
+	});
 
-		It("RHI Formats", [this]()
+	Describe("Test RHI Pixel Format", [this]()
+	{
+		It("RHI Target Formats", [this]()
 		{
 			bool bResult = RunOnRenderThreadSynchronous(FRHITextureTests::Test_RHIFormats);
-			TestEqual("RHI Formats failed", bResult, 1);
+			TestEqual("RHI Target Formats failed", bResult, 1);
 		});
+	});
 
+	Describe("Test RHI Resource Update", [this]()
+	{
 		It("RHI Update Texture2D", [this]()
 		{
 			bool bResult = RunOnRenderThreadSynchronous(FRHITextureTests::Test_UpdateTexture2D);
