Fixed compiler error in buffer tests.

diff --git a/Engine/Plugins/Tests/RHITests/Source/RHITests/Public/RHIBufferTests.h b/Engine/Plugins/Tests/RHITests/Source/RHITests/Public/RHIBufferTests.h
--- a/Engine/Plugins/Tests/RHITests/Source/RHITests/Public/RHIBufferTests.h
+++ b/Engine/Plugins/Tests/RHITests/Source/RHITests/Public/RHIBufferTests.h
@@ -37,13 +37,11 @@ class FRHIBufferTests
 
 		while (NumTestsLaunched < NumTests)
 		{
-			FRHICommandList* RHICmdListUpload = new FRHICommandList(FRHIGPUMask::All());
+			FRHICommandList* RHICmdListUpload = new FRHICommandList(FRHIGPUMask::All(), FRHICommandList::ERecordingThread::Any);
 
 			const int32 NumTestsInTask = FMath::Min(NumTestsPerTask, NumTests - NumTestsLaunched);
 
 #if 1
-			RHICmdListUpload->DisallowBypass();
-
 			FGraphEventRef Event = FFunctionGraphTask::CreateAndDispatchWhenReady(
 				[TestLambda, RHICmdListUpload, NumTestsLaunched, NumTestsInTask](ENamedThreads::Type, const FGraphEventRef&)
 			{
