Add RHITest to exercise MetalRHI fix for multiple slice locks on the same texture at the same time.

#rnx

diff --git a/Engine/Plugins/Tests/RHITests/Source/RHITests/Private/RHITests.spec.cpp b/Engine/Plugins/Tests/RHITests/Source/RHITests/Private/RHITests.spec.cpp
--- a/Engine/Plugins/Tests/RHITests/Source/RHITests/Private/RHITests.spec.cpp
+++ b/Engine/Plugins/Tests/RHITests/Source/RHITests/Private/RHITests.spec.cpp
@@ -67,6 +67,12 @@ void FAutomationRHITest::Define()
 			bool bResult = RunOnRenderThreadSynchronous(FRHITextureTests::Test_UpdateTexture2D);
 			TestEqual("RHI Update Texture2D failed", bResult, 1);
 		});
+		
+		It("RHI Multiple Lock Texture2D", [this]()
+		{
+			bool bResult = RunOnRenderThreadSynchronous(FRHITextureTests::Test_MultipleLockTexture2D);
+			TestEqual("RHI Multiple Lock Texture2D failed", bResult, 1);
+		});
 	});
 
 	Describe("Test RHI Copy", [this]
@@ -77,4 +83,4 @@ void FAutomationRHITest::Define()
 			TestEqual("RHICopyTexture", bResult, 1);
 		});
 	});
-}
\ No newline at end of file
+}
diff --git a/Engine/Plugins/Tests/RHITests/Source/RHITests/Private/RHITestsModule.cpp b/Engine/Plugins/Tests/RHITests/Source/RHITests/Private/RHITestsModule.cpp
--- a/Engine/Plugins/Tests/RHITests/Source/RHITests/Private/RHITestsModule.cpp
+++ b/Engine/Plugins/Tests/RHITests/Source/RHITests/Private/RHITestsModule.cpp
@@ -34,6 +34,7 @@ static bool RunTests_RenderThread(FRHICommandListImmediate& RHICmdList)
 		RUN_TEST(FRHITextureTests::Test_RHIClearUAV_Texture2D(RHICmdList));
 		RUN_TEST(FRHITextureTests::Test_RHIClearUAV_Texture3D(RHICmdList));
 		RUN_TEST(FRHITextureTests::Test_UpdateTexture2D(RHICmdList));
+		RUN_TEST(FRHITextureTests::Test_MultipleLockTexture2D(RHICmdList));
 	}
 
 	{
diff --git a/Engine/Plugins/Tests/RHITests/Source/RHITests/Public/RHITextureTests.h b/Engine/Plugins/Tests/RHITests/Source/RHITests/Public/RHITextureTests.h
--- a/Engine/Plugins/Tests/RHITests/Source/RHITests/Public/RHITextureTests.h
+++ b/Engine/Plugins/Tests/RHITests/Source/RHITests/Public/RHITextureTests.h
@@ -1111,6 +1111,114 @@ public:
 
 		return bResult;
 	}
+	
+	static bool Test_MultipleLockTexture2D(FRHICommandListImmediate& RHICmdList)
+	{
+		// Test the RHI can handle multiple simultaneous locks across texture array slices
+		const uint32 ArrayLength = 4;
+		const uint32 Size = 8;
+		const FIntPoint SliceDim = FIntPoint(Size, Size);
+		FRHITextureCreateDesc DescArray = FRHITextureCreateDesc::Create2DArray(TEXT("Multiple Texture Slice Lock"), SliceDim, ArrayLength, PF_R8G8B8A8);
+		
+		FTexture2DArrayRHIRef Texture_SingleLock = RHICreateTexture(DescArray);
+		FTexture2DArrayRHIRef Texture_MultipleLock = RHICreateTexture(DescArray);
+		
+		// Create Reasonable Test Data that can be seen in the GPU debugger if required
+		const auto WriteSliceTestData = [=](uint8* pDestData, uint32 DestStride, uint32 SliceIdx)
+		{
+			const uint32 c = 0xffffffff;
+			uint32 const Data[] =
+			{
+				0, 0, 0, c, c, 0, 0, 0,
+				0, 0, 0, c, c, 0, 0, 0,
+				0, 0, 0, c, c, 0, 0, 0,
+				c, c, c, c, c, c, c, c,
+				c, c, c, c, c, c, c, c,
+				0, 0, 0, c, c, 0, 0, 0,
+				0, 0, 0, c, c, 0, 0, 0,
+				0, 0, 0, c, c, 0, 0, 0,
+			};
+			
+			uint32 const TintMask[] = {0xff0000ff, 0xff00ff00, 0xffff0000, 0xffffffff};
+			
+			const uint32 SrcStride = sizeof(uint32) * Size;
+			uint8* pSrcData = (uint8*)Data;
+		
+			for (int32 Y = 0; Y < SliceDim.Y; ++Y)
+			{
+				uint32* pDestRowData = (uint32*)(pDestData + (Y * DestStride));
+				uint32* pSrcRowData = (uint32*)(pSrcData + (Y * SrcStride));
+
+				for (int32 X = 0; X < SliceDim.X;++X)
+				{
+					pDestRowData[X] = pSrcRowData[X] & TintMask[SliceIdx];
+				}
+			}
+		};
+		
+		// Lock, write data, then unlock each slice in turn
+		{
+			for (uint32 SliceIdx = 0; SliceIdx < ArrayLength;++SliceIdx)
+			{
+				uint32 DestStride = 0;
+				uint8* pDestData = (uint8*)RHICmdList.LockTexture2DArray(Texture_SingleLock, SliceIdx, 0, RLM_WriteOnly, DestStride, false);
+				WriteSliceTestData(pDestData, DestStride, SliceIdx);
+				RHICmdList.UnlockTexture2DArray(Texture_SingleLock, SliceIdx, 0, false);
+			}
+		}
+		
+		// Lock all slices, write data, then unlock all slices as individual steps
+		{
+			uint8* DataPtrs[ArrayLength] = {};
+			uint32 SliceStrides[ArrayLength] = {};
+			
+			for (uint32 SliceIdx = 0; SliceIdx < ArrayLength;++SliceIdx)
+			{
+				DataPtrs[SliceIdx] = (uint8*)RHICmdList.LockTexture2DArray(Texture_MultipleLock, SliceIdx, 0, RLM_WriteOnly, SliceStrides[SliceIdx], false);
+			}
+			
+			for(uint32 SliceIdx = 0;SliceIdx < ArrayLength;++SliceIdx)
+			{
+				WriteSliceTestData(DataPtrs[SliceIdx], SliceStrides[SliceIdx], SliceIdx);
+			}
+			
+			for (uint32 SliceIdx = 0; SliceIdx < ArrayLength;++SliceIdx)
+			{
+				RHICmdList.UnlockTexture2DArray(Texture_MultipleLock, SliceIdx, 0, false);
+			}
+		}
+		
+		RHICmdList.ImmediateFlush(EImmediateFlushType::FlushRHIThreadFlushResources);
+		
+		// Compare the two textures
+		{
+			for (uint32 SliceIdx = 0; SliceIdx < ArrayLength;++SliceIdx)
+			{
+				uint32 StrideSingleLock = 0;
+				uint8* pTextureSingleLockData = (uint8*)RHICmdList.LockTexture2DArray(Texture_SingleLock, SliceIdx, 0, RLM_ReadOnly, StrideSingleLock, false);
+				uint32 StrideMultiLock = 0;
+				uint8* pTextureMultiLockData = (uint8*)RHICmdList.LockTexture2DArray(Texture_MultipleLock, SliceIdx, 0, RLM_ReadOnly, StrideMultiLock, false);
+				
+				check(StrideSingleLock == StrideMultiLock && StrideSingleLock != 0);
+
+				for (int32 Y = 0; Y < SliceDim.Y; ++Y)
+				{
+					uint8* pRowDataSingle = (uint8*)(pTextureSingleLockData + (Y * StrideSingleLock));
+					uint8* pRowDataMulti = (uint8*)(pTextureMultiLockData + (Y * StrideMultiLock));
+
+					if(FMemory::Memcmp(pRowDataSingle, pRowDataMulti, StrideSingleLock) != 0)
+					{
+						return false;
+					}
+				}
+				
+				RHICmdList.UnlockTexture2DArray(Texture_SingleLock, SliceIdx, 0, false);
+				RHICmdList.UnlockTexture2DArray(Texture_MultipleLock, SliceIdx, 0, false);
+			}
+		}
+		
+		return true;
+	}
 };
 
-PRAGMA_ENABLE_OPTIMIZATION
\ No newline at end of file
+PRAGMA_ENABLE_OPTIMIZATION
