Fix RHI unit test failure caused by the new PF_R9G9B9EXP5 format not being taken into account.

#rnx

diff --git a/Engine/Plugins/Tests/RHITests/Source/RHITests/Public/RHITextureTests.h b/Engine/Plugins/Tests/RHITests/Source/RHITests/Public/RHITextureTests.h
--- a/Engine/Plugins/Tests/RHITests/Source/RHITests/Public/RHITextureTests.h
+++ b/Engine/Plugins/Tests/RHITests/Source/RHITests/Public/RHITextureTests.h
@@ -633,6 +633,7 @@ public:
 			case PF_R32G32B32F:
 			case PF_R32G32B32_UINT:
 			case PF_R32G32B32_SINT:
+			case PF_R9G9B9EXP5:
 				break;
 
 			case PF_G8:
