Don't early-out if the texture data doesn't match during the MultipleLockTexture2D RHI unit test, because that would leave the textures locked, which causes problems for subsequent tests or other rendering code.

#rnx

diff --git a/Engine/Plugins/Tests/RHITests/Source/RHITests/Public/RHITextureTests.h b/Engine/Plugins/Tests/RHITests/Source/RHITests/Public/RHITextureTests.h
--- a/Engine/Plugins/Tests/RHITests/Source/RHITests/Public/RHITextureTests.h
+++ b/Engine/Plugins/Tests/RHITests/Source/RHITests/Public/RHITextureTests.h
@@ -1202,6 +1202,7 @@ public:
 				
 				check(StrideSingleLock == StrideMultiLock && StrideSingleLock != 0);
 
+				bool bDataMatches = true;
 				for (int32 Y = 0; Y < SliceDim.Y; ++Y)
 				{
 					uint8* pRowDataSingle = (uint8*)(pTextureSingleLockData + (Y * StrideSingleLock));
@@ -1209,12 +1210,18 @@ public:
 
 					if(FMemory::Memcmp(pRowDataSingle, pRowDataMulti, StrideSingleLock) != 0)
 					{
-						return false;
+						bDataMatches = false;
+						break;
 					}
 				}
 				
 				RHICmdList.UnlockTexture2DArray(Texture_SingleLock, SliceIdx, 0, false);
 				RHICmdList.UnlockTexture2DArray(Texture_MultipleLock, SliceIdx, 0, false);
+
+				if (!bDataMatches)
+				{
+					return false;
+				}
 			}
 		}
 		
