Fix the MultipleLockTexture2D RHI unit test

This test was trying to read back the contents of a texture using the Lock function, which is only supported on a few RHIs. The correct way to do it is to copy to a CPU-accessible texture and use MapStagingSurface on that. Added a check in the Vulkan implementation of RHILockTexture2DArray to catch invalid usage (RHILockTexture2D already had this check).

Also fixed a typo in a couple of RHI unit test messages, and made the MultipleLockTexture2D report success/failure to the log.

#rnx

diff --git a/Engine/Plugins/Tests/RHITests/Source/RHITests/Private/RHITestsModule.cpp b/Engine/Plugins/Tests/RHITests/Source/RHITests/Private/RHITestsModule.cpp
--- a/Engine/Plugins/Tests/RHITests/Source/RHITests/Private/RHITestsModule.cpp
+++ b/Engine/Plugins/Tests/RHITests/Source/RHITests/Private/RHITestsModule.cpp
@@ -61,11 +61,11 @@ void FRHITestsModule::RunAllTests()
 	{
 		if (RunOnRenderThreadSynchronous(RunTests_RenderThread))
 		{
-			UE_LOG(LogRHIUnitTestCommandlet, Display, TEXT("RHI unit tested completed. All tests passed."));
+			UE_LOG(LogRHIUnitTestCommandlet, Display, TEXT("RHI unit tests completed. All tests passed."));
 		}
 		else
 		{
-			UE_LOG(LogRHIUnitTestCommandlet, Error, TEXT("RHI unit tested completed. At least one test failed."));
+			UE_LOG(LogRHIUnitTestCommandlet, Error, TEXT("RHI unit tests completed. At least one test failed."));
 		}
 	}
 }
diff --git a/Engine/Plugins/Tests/RHITests/Source/RHITests/Public/RHITextureTests.h b/Engine/Plugins/Tests/RHITests/Source/RHITests/Public/RHITextureTests.h
--- a/Engine/Plugins/Tests/RHITests/Source/RHITests/Public/RHITextureTests.h
+++ b/Engine/Plugins/Tests/RHITests/Source/RHITests/Public/RHITextureTests.h
@@ -6,7 +6,7 @@
 #include <type_traits>
 #include "Math/PackedVector.h"
 
-PRAGMA_DISABLE_OPTIMIZATION
+//PRAGMA_DISABLE_OPTIMIZATION
 
 class FRHITextureTests
 {
@@ -1115,11 +1115,14 @@ public:
 	
 	static bool Test_MultipleLockTexture2D(FRHICommandListImmediate& RHICmdList)
 	{
+		FString TestName = TEXT("Test_MultipleLockTexture2D");
+
 		// Test the RHI can handle multiple simultaneous locks across texture array slices
 		const uint32 ArrayLength = 4;
 		const uint32 Size = 8;
 		const FIntPoint SliceDim = FIntPoint(Size, Size);
-		FRHITextureCreateDesc DescArray = FRHITextureCreateDesc::Create2DArray(TEXT("Multiple Texture Slice Lock"), SliceDim, ArrayLength, PF_R8G8B8A8);
+		const EPixelFormat Format = PF_R8G8B8A8;
+		FRHITextureCreateDesc DescArray = FRHITextureCreateDesc::Create2DArray(TEXT("Multiple Texture Slice Lock"), SliceDim, ArrayLength, Format).DetermineInititialState();
 		
 		FTexture2DArrayRHIRef Texture_SingleLock = RHICreateTexture(DescArray);
 		FTexture2DArrayRHIRef Texture_MultipleLock = RHICreateTexture(DescArray);
@@ -1189,44 +1192,82 @@ public:
 			}
 		}
 		
+		RHICmdList.Transition({
+			FRHITransitionInfo(Texture_SingleLock, DescArray.InitialState, ERHIAccess::CopySrc),
+			FRHITransitionInfo(Texture_MultipleLock, DescArray.InitialState, ERHIAccess::CopySrc),
+		});
+
 		RHICmdList.ImmediateFlush(EImmediateFlushType::FlushRHIThreadFlushResources);
-		
+
+		const FRHITextureCreateDesc Desc =
+			FRHITextureCreateDesc::Create2D(TEXT("FRHITextureTests_StagingTexture"), SliceDim.X, SliceDim.Y, Format)
+			.SetFlags(ETextureCreateFlags::CPUReadback)
+			.SetInitialState(ERHIAccess::CopyDest);
+
+		FTextureRHIRef SingleLockStagingTexture = RHICreateTexture(Desc);
+		FTextureRHIRef MultiLockStagingTexture = RHICreateTexture(Desc);
+
+		FRHICopyTextureInfo CopyInfo = {};
+		CopyInfo.Size = FIntVector(SliceDim.X, SliceDim.Y, 1);
+		CopyInfo.SourceMipIndex = 0;
+		CopyInfo.NumSlices = 1;
+		CopyInfo.NumMips = 1;
+
 		// Compare the two textures
+		for (uint32 SliceIdx = 0; SliceIdx < ArrayLength;++SliceIdx)
 		{
-			for (uint32 SliceIdx = 0; SliceIdx < ArrayLength;++SliceIdx)
-			{
-				uint32 StrideSingleLock = 0;
-				uint8* pTextureSingleLockData = (uint8*)RHICmdList.LockTexture2DArray(Texture_SingleLock, SliceIdx, 0, RLM_ReadOnly, StrideSingleLock, false);
-				uint32 StrideMultiLock = 0;
-				uint8* pTextureMultiLockData = (uint8*)RHICmdList.LockTexture2DArray(Texture_MultipleLock, SliceIdx, 0, RLM_ReadOnly, StrideMultiLock, false);
-				
-				check(StrideSingleLock == StrideMultiLock && StrideSingleLock != 0);
-
-				bool bDataMatches = true;
-				for (int32 Y = 0; Y < SliceDim.Y; ++Y)
-				{
-					uint8* pRowDataSingle = (uint8*)(pTextureSingleLockData + (Y * StrideSingleLock));
-					uint8* pRowDataMulti = (uint8*)(pTextureMultiLockData + (Y * StrideMultiLock));
+			CopyInfo.SourceSliceIndex = SliceIdx;
+			RHICmdList.CopyTexture(Texture_SingleLock, SingleLockStagingTexture, CopyInfo);
+			RHICmdList.CopyTexture(Texture_MultipleLock, MultiLockStagingTexture, CopyInfo);
 
-					if(FMemory::Memcmp(pRowDataSingle, pRowDataMulti, StrideSingleLock) != 0)
-					{
-						bDataMatches = false;
-						break;
-					}
-				}
-				
-				RHICmdList.UnlockTexture2DArray(Texture_SingleLock, SliceIdx, 0, false);
-				RHICmdList.UnlockTexture2DArray(Texture_MultipleLock, SliceIdx, 0, false);
+			RHICmdList.Transition({
+				FRHITransitionInfo(SingleLockStagingTexture, ERHIAccess::CopyDest, ERHIAccess::CPURead),
+				FRHITransitionInfo(MultiLockStagingTexture, ERHIAccess::CopyDest, ERHIAccess::CPURead)
+			});
+
+			FGPUFenceRHIRef GPUFence = RHICreateGPUFence(TEXT("ReadbackFence"));
+			RHICmdList.WriteGPUFence(GPUFence);
+
+			RHICmdList.SubmitCommandsAndFlushGPU();
+			RHICmdList.FlushResources();
+
+			int32 SingleLockStagingSurfaceWidth, SingleLockStagingSurfaceHeight;
+			void* SingleLockData;
+			RHICmdList.MapStagingSurface(SingleLockStagingTexture, GPUFence, SingleLockData, SingleLockStagingSurfaceWidth, SingleLockStagingSurfaceHeight);
+			check(SingleLockStagingSurfaceWidth >= SliceDim.X && SingleLockStagingSurfaceHeight >= SliceDim.Y);
+
+			int32 MultiLockStagingSurfaceWidth, MultiLockStagingSurfaceHeight;
+			void* MultiLockData;
+			RHICmdList.MapStagingSurface(MultiLockStagingTexture, GPUFence, MultiLockData, MultiLockStagingSurfaceWidth, MultiLockStagingSurfaceHeight);
+			check(MultiLockStagingSurfaceWidth >= SliceDim.X && MultiLockStagingSurfaceHeight >= SliceDim.Y);
 
-				if (!bDataMatches)
+			bool bDataMatches = true;
+			const uint8* SingleLockCursor = (const uint8*)SingleLockData;
+			const uint8* MultiLockCursor = (const uint8*)MultiLockData;
+			for (int32 Y = 0; Y < SliceDim.Y; ++Y)
+			{
+				if (FMemory::Memcmp(SingleLockCursor, MultiLockCursor, SliceDim.X * sizeof(uint32)) != 0)
 				{
-					return false;
+					bDataMatches = false;
+					break;
 				}
+				SingleLockCursor += SingleLockStagingSurfaceWidth * sizeof(uint32);
+				MultiLockCursor += MultiLockStagingSurfaceWidth * sizeof(uint32);
+			}
+
+			RHICmdList.UnmapStagingSurface(SingleLockStagingTexture);
+			RHICmdList.UnmapStagingSurface(MultiLockStagingTexture);
+
+			if (!bDataMatches)
+			{
+				UE_LOG(LogRHIUnitTestCommandlet, Display, TEXT("Test failed. \"%s\""), *TestName);
+				return false;
 			}
 		}
 		
+		UE_LOG(LogRHIUnitTestCommandlet, Display, TEXT("Test passed. \"%s\""), *TestName);
 		return true;
 	}
 };
 
-PRAGMA_ENABLE_OPTIMIZATION
+//PRAGMA_ENABLE_OPTIMIZATION
diff --git a/Engine/Source/Runtime/VulkanRHI/Private/VulkanTexture.cpp b/Engine/Source/Runtime/VulkanRHI/Private/VulkanTexture.cpp
--- a/Engine/Source/Runtime/VulkanRHI/Private/VulkanTexture.cpp
+++ b/Engine/Source/Runtime/VulkanRHI/Private/VulkanTexture.cpp
@@ -1203,6 +1203,9 @@ void* FVulkanDynamicRHI::RHILockTexture2DArray(FRHITexture2DArray* TextureRHI,ui
 		checkf(!*StagingBuffer, TEXT("Can't lock the same texture twice!"));
 	}
 
+	// No locks for read allowed yet
+	check(LockMode == RLM_WriteOnly);
+
 	uint32 BufferSize = 0;
 	DestStride = 0;
 	Texture->GetMipSize(MipIndex, BufferSize);
