Fix validation errors in RHI unit tests.

#rnx

diff --git a/Engine/Plugins/Tests/RHITests/Source/RHITests/Public/RHIBufferTests.h b/Engine/Plugins/Tests/RHITests/Source/RHITests/Public/RHIBufferTests.h
--- a/Engine/Plugins/Tests/RHITests/Source/RHITests/Public/RHIBufferTests.h
+++ b/Engine/Plugins/Tests/RHITests/Source/RHITests/Public/RHIBufferTests.h
@@ -398,7 +398,7 @@ public:
 			const int32 BufferSize = Align(Random % 65536, 16);
 			const int32 BufferStride = 4;
 
-			EBufferUsageFlags Usage = EBufferUsageFlags::VertexBuffer;
+			EBufferUsageFlags Usage = EBufferUsageFlags::VertexBuffer | EBufferUsageFlags::SourceCopy;
 
 			switch (Random % 3)
 			{
@@ -415,7 +415,7 @@ public:
 
 			FRHIResourceCreateInfo CreateInfo(TEXT("Buffer"));
 
-			TRefCountPtr<FRHIBuffer> Buffer = InRHICmdList.CreateBuffer(BufferSize, Usage, 0, ERHIAccess::VertexOrIndexBuffer, CreateInfo);
+			TRefCountPtr<FRHIBuffer> Buffer = InRHICmdList.CreateBuffer(BufferSize, Usage, 0, ERHIAccess::CopySrc, CreateInfo);
 
 			uint32* Data = (uint32*)InRHICmdList.LockBuffer(Buffer, 0, Buffer->GetSize(), RLM_WriteOnly);
 
diff --git a/Engine/Plugins/Tests/RHITests/Source/RHITests/Public/RHITextureTests.h b/Engine/Plugins/Tests/RHITests/Source/RHITests/Public/RHITextureTests.h
--- a/Engine/Plugins/Tests/RHITests/Source/RHITests/Public/RHITextureTests.h
+++ b/Engine/Plugins/Tests/RHITests/Source/RHITests/Public/RHITextureTests.h
@@ -1263,6 +1263,13 @@ public:
 				UE_LOG(LogRHIUnitTestCommandlet, Display, TEXT("Test failed. \"%s\""), *TestName);
 				return false;
 			}
+
+			// Put the staging textures back into CopyDest for the next iteration of the loop.
+			RHICmdList.Transition({
+				FRHITransitionInfo(SingleLockStagingTexture, ERHIAccess::CPURead, ERHIAccess::CopyDest),
+				FRHITransitionInfo(MultiLockStagingTexture, ERHIAccess::CPURead, ERHIAccess::CopyDest)
+			});
+
 		}
 		
 		UE_LOG(LogRHIUnitTestCommandlet, Display, TEXT("Test passed. \"%s\""), *TestName);
