Refactored UpdateTexture2D, UpdateFromBufferTexture2D, and UpdateTexture3D to take a non-immediate command list.
 - Added explicit SwitchPipeline calls to async command lists.
 - Moved pipeline push / pop calls from platform implementations to RHI command list.

diff --git a/Engine/Plugins/Media/BinkMedia/Source/BinkMediaPlayer/Private/binktexturesRHI.cpp b/Engine/Plugins/Media/BinkMedia/Source/BinkMediaPlayer/Private/binktexturesRHI.cpp
--- a/Engine/Plugins/Media/BinkMedia/Source/BinkMediaPlayer/Private/binktexturesRHI.cpp
+++ b/Engine/Plugins/Media/BinkMedia/Source/BinkMediaPlayer/Private/binktexturesRHI.cpp
@@ -491,11 +491,11 @@ static void Draw_textures(BINKTEXTURES* ptextures, BINKSHADERS* pshaders, void*
 				FUpdateTextureRegion2D region_YAH(0, 0, 0, 0, bb->YABufferWidth, bb->YABufferHeight);
 				FUpdateTextureRegion2D region_cRcB(0, 0, 0, 0, bb->cRcBBufferWidth, bb->cRcBBufferHeight);
 
-				GDynamicRHI->UpdateTexture2D_RenderThread(RHICmdList, textures->Ytexture[frame_num], 0, region_YAH, bp_src->YPlane.BufferPitch, (uint8*)bp_src->YPlane.Buffer);
-				GDynamicRHI->UpdateTexture2D_RenderThread(RHICmdList, textures->cRtexture[frame_num], 0, region_cRcB, bp_src->cRPlane.BufferPitch, (uint8*)bp_src->cRPlane.Buffer);
-				GDynamicRHI->UpdateTexture2D_RenderThread(RHICmdList, textures->cBtexture[frame_num], 0, region_cRcB, bp_src->cBPlane.BufferPitch, (uint8*)bp_src->cBPlane.Buffer);
-				if (hasAPlane) GDynamicRHI->UpdateTexture2D_RenderThread(RHICmdList, textures->Atexture[frame_num], 0, region_YAH, bp_src->APlane.BufferPitch, (uint8*)bp_src->APlane.Buffer);
-				if (hasHPlane) GDynamicRHI->UpdateTexture2D_RenderThread(RHICmdList, textures->Htexture[frame_num], 0, region_YAH, bp_src->HPlane.BufferPitch, (uint8*)bp_src->HPlane.Buffer);
+				GDynamicRHI->RHIUpdateTexture2D(RHICmdList, textures->Ytexture[frame_num], 0, region_YAH, bp_src->YPlane.BufferPitch, (uint8*)bp_src->YPlane.Buffer);
+				GDynamicRHI->RHIUpdateTexture2D(RHICmdList, textures->cRtexture[frame_num], 0, region_cRcB, bp_src->cRPlane.BufferPitch, (uint8*)bp_src->cRPlane.Buffer);
+				GDynamicRHI->RHIUpdateTexture2D(RHICmdList, textures->cBtexture[frame_num], 0, region_cRcB, bp_src->cBPlane.BufferPitch, (uint8*)bp_src->cBPlane.Buffer);
+				if (hasAPlane) GDynamicRHI->RHIUpdateTexture2D(RHICmdList, textures->Atexture[frame_num], 0, region_YAH, bp_src->APlane.BufferPitch, (uint8*)bp_src->APlane.Buffer);
+				if (hasHPlane) GDynamicRHI->RHIUpdateTexture2D(RHICmdList, textures->Htexture[frame_num], 0, region_YAH, bp_src->HPlane.BufferPitch, (uint8*)bp_src->HPlane.Buffer);
 			}
 
 			FGraphicsPipelineStateInitializer GraphicsPSOInit;
diff --git a/Engine/Plugins/Tests/RHITests/Source/RHITests/Private/RHITests.spec.cpp b/Engine/Plugins/Tests/RHITests/Source/RHITests/Private/RHITests.spec.cpp
--- a/Engine/Plugins/Tests/RHITests/Source/RHITests/Private/RHITests.spec.cpp
+++ b/Engine/Plugins/Tests/RHITests/Source/RHITests/Private/RHITests.spec.cpp
@@ -64,8 +64,8 @@ void FAutomationRHITest::Define()
 	{
 		It("RHI Update Texture2D", [this]()
 		{
-			bool bResult = RunOnRenderThreadSynchronous(FRHITextureTests::Test_UpdateTexture2D);
-			TestEqual("RHI Update Texture2D failed", bResult, 1);
+			bool bResult = RunOnRenderThreadSynchronous(FRHITextureTests::Test_UpdateTexture);
+			TestEqual("RHI Update Texture failed", bResult, 1);
 		});
 		
 		It("RHI Multiple Lock Texture2D", [this]()
diff --git a/Engine/Plugins/Tests/RHITests/Source/RHITests/Private/RHITestsModule.cpp b/Engine/Plugins/Tests/RHITests/Source/RHITests/Private/RHITestsModule.cpp
--- a/Engine/Plugins/Tests/RHITests/Source/RHITests/Private/RHITestsModule.cpp
+++ b/Engine/Plugins/Tests/RHITests/Source/RHITests/Private/RHITestsModule.cpp
@@ -33,7 +33,7 @@ static bool RunTests_RenderThread(FRHICommandListImmediate& RHICmdList)
 	{
 		RUN_TEST(FRHITextureTests::Test_RHIClearUAV_Texture2D(RHICmdList));
 		RUN_TEST(FRHITextureTests::Test_RHIClearUAV_Texture3D(RHICmdList));
-		RUN_TEST(FRHITextureTests::Test_UpdateTexture2D(RHICmdList));
+		RUN_TEST(FRHITextureTests::Test_UpdateTexture(RHICmdList));
 		RUN_TEST(FRHITextureTests::Test_MultipleLockTexture2D(RHICmdList));
 	}
 
diff --git a/Engine/Plugins/Tests/RHITests/Source/RHITests/Public/RHITextureTests.h b/Engine/Plugins/Tests/RHITests/Source/RHITests/Public/RHITextureTests.h
--- a/Engine/Plugins/Tests/RHITests/Source/RHITests/Public/RHITextureTests.h
+++ b/Engine/Plugins/Tests/RHITests/Source/RHITests/Public/RHITextureTests.h
@@ -432,8 +432,29 @@ public:
 	typedef FDataSource< uint64,		1, false >	FDataSource64x1UInt;
 	typedef FDataSource< FFloat3Packed,	1, true >	FDataSource11_11_10F;
 
+	enum class ETextureUpdateType
+	{
+		Texture2D,
+		TextureFromBuffer,
+		Texture3D
+	};
+
+	static const TCHAR* GetName(ETextureUpdateType UpdateType)
+	{
+		switch (UpdateType)
+		{
+		case ETextureUpdateType::Texture2D:
+			return TEXT("Texture2D");
+		case ETextureUpdateType::TextureFromBuffer:
+			return TEXT("TextureFromBuffer");
+		case ETextureUpdateType::Texture3D:
+			return TEXT("Texture3D");
+		}
+		return TEXT("");
+	}
+
 	template <typename SourceType>
-	static bool Test_UpdateTexture2D_Impl(FRHICommandListImmediate& RHICmdList, FString TestName, FRHITexture2D* Texture, const FUpdateTextureRegion2D& Region, uint32 SourcePitch, const uint8* SourceData, const uint8* ZeroData)
+	static bool Test_UpdateTexture_Impl(FRHICommandListImmediate& RHICmdList, FString TestName, FRHITexture* Texture, const FUpdateTextureRegion2D& Region, uint32 SourcePitch, const uint8* SourceData, const uint8* ZeroData, ETextureUpdateType UpdateType)
 	{
 		bool bResult = true;
 		
@@ -470,24 +491,62 @@ public:
 			return true;
 		};
 
-		// clear to zero
 		RHICmdList.Transition(FRHITransitionInfo(Texture, ERHIAccess::Unknown, ERHIAccess::UAVCompute));
-		FUpdateTextureRegion2D ZeroRegion = { 0, 0, 0, 0, Texture->GetSizeX(), Texture->GetSizeY() };
-		RHICmdList.UpdateTexture2D(Texture, 0, ZeroRegion, SourcePitch, ZeroData);
+
+		// clear to zero
+		if (UpdateType == ETextureUpdateType::Texture3D)
+		{
+			for (uint32 DepthSlice = 0; DepthSlice < Texture->GetSizeZ(); ++DepthSlice)
+			{
+				FUpdateTextureRegion3D ZeroRegion = { 0, 0, DepthSlice, 0, 0, 0, Texture->GetSizeX(), Texture->GetSizeY(), 1 };
+				RHICmdList.UpdateTexture3D(Texture, 0, ZeroRegion, SourcePitch, SourcePitch * ZeroRegion.Height, ZeroData);
+			}
+		}
+		else
+		{
+			FUpdateTextureRegion2D ZeroRegion = { 0, 0, 0, 0, Texture->GetSizeX(), Texture->GetSizeY() };
+			RHICmdList.UpdateTexture2D(Texture, 0, ZeroRegion, SourcePitch, ZeroData);
+		}
+
 		RHICmdList.Transition(FRHITransitionInfo(Texture, ERHIAccess::UAVCompute, ERHIAccess::CopySrc));
+
 		bResult &= VerifyTextureContents(*FString::Printf(TEXT("%s - clear whole resource to zero"), *TestName), RHICmdList, Texture,
-		[&](void* Ptr, uint32 MipWidth, uint32 MipHeight, uint32 Width, uint32 Height, uint32 CurrentMipIndex, uint32 CurrentSliceIndex)
+			[&](void* Ptr, uint32 MipWidth, uint32 MipHeight, uint32 Width, uint32 Height, uint32 CurrentMipIndex, uint32 CurrentSliceIndex)
 		{
 			return VerifyMip(Ptr, MipWidth, MipHeight, Width, Height, CurrentMipIndex, CurrentSliceIndex, true);
 		});
-		RHICmdList.Transition(FRHITransitionInfo(Texture, ERHIAccess::CopySrc, ERHIAccess::UAVCompute));
 
-		// update the texture
-		RHICmdList.UpdateTexture2D(Texture, 0, Region, SourcePitch, SourceData);
+		if (UpdateType == ETextureUpdateType::Texture2D)
+		{
+			RHICmdList.UpdateTexture2D(Texture, 0, Region, SourcePitch, SourceData);
+		}
+		else if (UpdateType == ETextureUpdateType::TextureFromBuffer)
+		{
+			const int32 BufferSize = Region.Width * Region.Height;
+
+			FRHIResourceCreateInfo CreateInfo(TEXT("SourceUpdateBuffer"));
+			TRefCountPtr<FRHIBuffer> Buffer = RHICmdList.CreateBuffer(BufferSize, EBufferUsageFlags::ShaderResource | EBufferUsageFlags::Static | EBufferUsageFlags::KeepCPUAccessible, SourceType::ElementSize, ERHIAccess::SRVCompute, CreateInfo);
+			void* Data = RHICmdList.LockBuffer(Buffer, 0, BufferSize, RLM_WriteOnly);
+			FMemory::Memcpy(Data, SourceData, BufferSize);
+			RHICmdList.UnlockBuffer(Buffer);
+
+			RHICmdList.UpdateFromBufferTexture2D(Texture, 0, Region, SourcePitch, Buffer, 0);
+		}
+		else
+		{
+			check(UpdateType == ETextureUpdateType::Texture3D);
+
+			for (uint32 DepthSlice = 0; DepthSlice < Texture->GetSizeZ(); ++DepthSlice)
+			{
+				FUpdateTextureRegion3D Region3D = { Region.DestX, Region.DestY, DepthSlice, Region.SrcX, Region.SrcY, 0, Region.Width, Region.Height, 1 };
+				RHICmdList.UpdateTexture3D(Texture, 0, Region3D, SourcePitch, SourcePitch * Texture->GetSizeY(), SourceData);
+			}
+		}
+
 		RHICmdList.Transition(FRHITransitionInfo(Texture, ERHIAccess::UAVCompute, ERHIAccess::CopySrc));
 
 		bResult &= VerifyTextureContents(*FString::Printf(TEXT("%s - update (%d,%d -> %d,%d)"), *TestName, Region.SrcX, Region.SrcY, Region.DestX, Region.DestY), RHICmdList, Texture,
-		[&](void* Ptr, uint32 MipWidth, uint32 MipHeight, uint32 Width, uint32 Height, uint32 CurrentMipIndex, uint32 CurrentSliceIndex)
+			[&](void* Ptr, uint32 MipWidth, uint32 MipHeight, uint32 Width, uint32 Height, uint32 CurrentMipIndex, uint32 CurrentSliceIndex)
 		{
 			return VerifyMip(Ptr, MipWidth, MipHeight, Width, Height, CurrentMipIndex, CurrentSliceIndex, false);
 		});
@@ -498,7 +557,7 @@ public:
 	}
 
 	template <typename SourceType>
-	static bool Test_UpdateTexture2D_Impl(FRHICommandListImmediate& RHICmdList, EPixelFormat Format)
+	static bool Test_UpdateTexture_Impl(FRHICommandListImmediate& RHICmdList, EPixelFormat Format, ETextureUpdateType UpdateType)
 	{
 		const uint32 TextureWidth = 128;
 		const uint32 TextureHeight = 128;
@@ -518,22 +577,39 @@ public:
 		FMemory::Memset(ZeroData, 0, SrcDataSize);
 
 		bool bResult = true;
-		FString TestName = FString::Printf(TEXT("Test_UpdateTexture2D (%s)"), GPixelFormats[Format].Name);
+		FString TestName = FString::Printf(TEXT("Test_UpdateTexture (%s) (%s)"), GPixelFormats[Format].Name, GetName(UpdateType));
 		{
 			if (!GPixelFormats[Format].Supported)
 			{
 				UE_LOG(LogRHIUnitTestCommandlet, Display, TEXT("Test skipped (format not supported). \"%s\""), *TestName);
 				return true;
 			}
-			if (!UE::PixelFormat::HasCapabilities(Format, EPixelFormatCapabilities::Texture2D))
+			FRHITextureCreateDesc Desc;
+			
+			if (UpdateType == ETextureUpdateType::Texture3D)
 			{
-				UE_LOG(LogRHIUnitTestCommandlet, Display, TEXT("Test skipped (format not supported as Texture2D). \"%s\""), *TestName);
-				return true;
+				const int32 TextureDepth = 4;
+
+				if (!UE::PixelFormat::HasCapabilities(Format, EPixelFormatCapabilities::Texture3D))
+				{
+					UE_LOG(LogRHIUnitTestCommandlet, Display, TEXT("Test skipped (format not supported as Texture3D). \"%s\""), *TestName);
+					return true;
+				}
+
+				Desc = FRHITextureCreateDesc::Create3D(*TestName, FIntVector(TextureWidth, TextureHeight, TextureDepth), Format);
+			}
+			else
+			{
+				if (!UE::PixelFormat::HasCapabilities(Format, EPixelFormatCapabilities::Texture2D))
+				{
+					UE_LOG(LogRHIUnitTestCommandlet, Display, TEXT("Test skipped (format not supported as Texture2D). \"%s\""), *TestName);
+					return true;
+				}
+
+				Desc = FRHITextureCreateDesc::Create2D(*TestName, TextureWidth, TextureHeight, Format);
 			}
 
-			const FRHITextureCreateDesc Desc =
-				FRHITextureCreateDesc::Create2D(*TestName, TextureWidth, TextureHeight, Format)
-				.SetFlags(ETextureCreateFlags::ShaderResource);
+			Desc.SetFlags(ETextureCreateFlags::ShaderResource);
 
 			FTexture2DRHIRef Texture = RHICreateTexture(Desc);
 			if (Texture == nullptr)
@@ -561,7 +637,7 @@ public:
 				for (uint32 DestColumn = 0; DestColumn < 3; DestColumn++)
 				{
 					Region.DestX = AxisSlotToCoord(DestColumn, TextureWidth, UpdateWidth);
-					bResult &= Test_UpdateTexture2D_Impl<SourceType>(RHICmdList, TestName, Texture, Region, UpdateDataPitch, UpdateData, ZeroData);
+					bResult &= Test_UpdateTexture_Impl<SourceType>(RHICmdList, TestName, Texture, Region, UpdateDataPitch, UpdateData, ZeroData, UpdateType);
 				}
 			}
 
@@ -573,7 +649,7 @@ public:
 				for (uint32 SrcColumn = 0; SrcColumn < 3; SrcColumn++)
 				{
 					Region.SrcX = AxisSlotToCoord(SrcColumn, SrcDataWidth, UpdateWidth);
-					bResult &= Test_UpdateTexture2D_Impl<SourceType>(RHICmdList, TestName, Texture, Region, UpdateDataPitch, UpdateData, ZeroData);
+					bResult &= Test_UpdateTexture_Impl<SourceType>(RHICmdList, TestName, Texture, Region, UpdateDataPitch, UpdateData, ZeroData, UpdateType);
 				}
 			}
 #endif
@@ -582,7 +658,7 @@ public:
 		return bResult;
 	}
 
-	static bool Test_UpdateTexture2D(FRHICommandListImmediate& RHICmdList)
+	static bool Test_UpdateTexture_Impl(FRHICommandListImmediate& RHICmdList, ETextureUpdateType UpdateType)
 	{
 		bool bResult = true;
 		auto RunTestOnFormat = [&](EPixelFormat Format)
@@ -640,92 +716,92 @@ public:
 			case PF_A8:
 			case PF_L8:
 			case PF_R8:
-				RUN_TEST(Test_UpdateTexture2D_Impl<FDataSource8x1>(RHICmdList, Format));
+				RUN_TEST(Test_UpdateTexture_Impl<FDataSource8x1>(RHICmdList, Format, UpdateType));
 				break;
 			case PF_R8_UINT:
 			case PF_R8_SINT:
-				RUN_TEST(Test_UpdateTexture2D_Impl<FDataSource8x1UInt>(RHICmdList, Format));
+				RUN_TEST(Test_UpdateTexture_Impl<FDataSource8x1UInt>(RHICmdList, Format, UpdateType));
 				break;
 			case PF_V8U8:
 			case PF_R8G8:
-				RUN_TEST(Test_UpdateTexture2D_Impl<FDataSource8x2>(RHICmdList, Format));
+				RUN_TEST(Test_UpdateTexture_Impl<FDataSource8x2>(RHICmdList, Format, UpdateType));
 				break;
 			case PF_R8G8_UINT:
-				RUN_TEST(Test_UpdateTexture2D_Impl<FDataSource8x2UInt>(RHICmdList, Format));
+				RUN_TEST(Test_UpdateTexture_Impl<FDataSource8x2UInt>(RHICmdList, Format, UpdateType));
 				break;
 			case PF_B8G8R8A8:
 			case PF_R8G8B8A8:
 			case PF_A8R8G8B8:
 			case PF_R8G8B8A8_SNORM:
-				RUN_TEST(Test_UpdateTexture2D_Impl<FDataSource8x4>(RHICmdList, Format));
+				RUN_TEST(Test_UpdateTexture_Impl<FDataSource8x4>(RHICmdList, Format, UpdateType));
 				break;
 			case PF_R8G8B8A8_UINT:
-				RUN_TEST(Test_UpdateTexture2D_Impl<FDataSource8x4UInt>(RHICmdList, Format));
+				RUN_TEST(Test_UpdateTexture_Impl<FDataSource8x4UInt>(RHICmdList, Format, UpdateType));
 				break;
 			case PF_G16:
 			case PF_R5G6B5_UNORM:	// only for the 16-bit pixel data
 			case PF_B5G5R5A1_UNORM:	// only for the 16-bit pixel data
-				RUN_TEST(Test_UpdateTexture2D_Impl<FDataSource16x1>(RHICmdList, Format));
+				RUN_TEST(Test_UpdateTexture_Impl<FDataSource16x1>(RHICmdList, Format, UpdateType));
 				break;
 			case PF_R16_UINT:
 			case PF_R16_SINT:
-				RUN_TEST(Test_UpdateTexture2D_Impl<FDataSource16x1UInt>(RHICmdList, Format));
+				RUN_TEST(Test_UpdateTexture_Impl<FDataSource16x1UInt>(RHICmdList, Format, UpdateType));
 				break;
 			case PF_R16F:
 			case PF_R16F_FILTER:
-				RUN_TEST(Test_UpdateTexture2D_Impl<FDataSource16x1F>(RHICmdList, Format));
+				RUN_TEST(Test_UpdateTexture_Impl<FDataSource16x1F>(RHICmdList, Format, UpdateType));
 				break;
 			case PF_G16R16:
 			case PF_G16R16_SNORM:
-				RUN_TEST(Test_UpdateTexture2D_Impl<FDataSource16x2>(RHICmdList, Format));
+				RUN_TEST(Test_UpdateTexture_Impl<FDataSource16x2>(RHICmdList, Format, UpdateType));
 				break;
 			case PF_R16G16_UINT:
-				RUN_TEST(Test_UpdateTexture2D_Impl<FDataSource16x2UInt>(RHICmdList, Format));
+				RUN_TEST(Test_UpdateTexture_Impl<FDataSource16x2UInt>(RHICmdList, Format, UpdateType));
 				break;
 			case PF_G16R16F:
 			case PF_G16R16F_FILTER:
-				RUN_TEST(Test_UpdateTexture2D_Impl<FDataSource16x2F>(RHICmdList, Format));
+				RUN_TEST(Test_UpdateTexture_Impl<FDataSource16x2F>(RHICmdList, Format, UpdateType));
 				break;
 			case PF_A16B16G16R16:
 			case PF_R16G16B16A16_UNORM:
 			case PF_R16G16B16A16_SNORM:
-				RUN_TEST(Test_UpdateTexture2D_Impl<FDataSource16x4>(RHICmdList, Format));
+				RUN_TEST(Test_UpdateTexture_Impl<FDataSource16x4>(RHICmdList, Format, UpdateType));
 				break;
 			case PF_R16G16B16A16_SINT:
 			case PF_R16G16B16A16_UINT:
-				RUN_TEST(Test_UpdateTexture2D_Impl<FDataSource16x4UInt>(RHICmdList, Format));
+				RUN_TEST(Test_UpdateTexture_Impl<FDataSource16x4UInt>(RHICmdList, Format, UpdateType));
 				break;
 			case PF_R32_SINT:
 			case PF_R32_UINT:
-				RUN_TEST(Test_UpdateTexture2D_Impl<FDataSource32x1UInt>(RHICmdList, Format));
+				RUN_TEST(Test_UpdateTexture_Impl<FDataSource32x1UInt>(RHICmdList, Format, UpdateType));
 				break;
 			case PF_R32_FLOAT:
-				RUN_TEST(Test_UpdateTexture2D_Impl<FDataSource32x1F>(RHICmdList, Format));
+				RUN_TEST(Test_UpdateTexture_Impl<FDataSource32x1F>(RHICmdList, Format, UpdateType));
 				break;
 			case PF_G32R32F:
-				RUN_TEST(Test_UpdateTexture2D_Impl<FDataSource32x2F>(RHICmdList, Format));
+				RUN_TEST(Test_UpdateTexture_Impl<FDataSource32x2F>(RHICmdList, Format, UpdateType));
 				break;
 			case PF_R32G32_UINT:
-				RUN_TEST(Test_UpdateTexture2D_Impl<FDataSource32x2UInt>(RHICmdList, Format));
+				RUN_TEST(Test_UpdateTexture_Impl<FDataSource32x2UInt>(RHICmdList, Format, UpdateType));
 				break;
 			case PF_A32B32G32R32F:
-				RUN_TEST(Test_UpdateTexture2D_Impl<FDataSource32x4F>(RHICmdList, Format));
+				RUN_TEST(Test_UpdateTexture_Impl<FDataSource32x4F>(RHICmdList, Format, UpdateType));
 				break;
 			case PF_R32G32B32A32_UINT:
-				RUN_TEST(Test_UpdateTexture2D_Impl<FDataSource32x4UInt>(RHICmdList, Format));
+				RUN_TEST(Test_UpdateTexture_Impl<FDataSource32x4UInt>(RHICmdList, Format, UpdateType));
 				break;
 			case PF_R64_UINT:
-				RUN_TEST(Test_UpdateTexture2D_Impl<FDataSource64x1UInt>(RHICmdList, Format));
+				RUN_TEST(Test_UpdateTexture_Impl<FDataSource64x1UInt>(RHICmdList, Format, UpdateType));
 				break;
 			case PF_FloatRGB:
 			case PF_FloatR11G11B10:
-				RUN_TEST(Test_UpdateTexture2D_Impl<FDataSource11_11_10F>(RHICmdList, Format));
+				RUN_TEST(Test_UpdateTexture_Impl<FDataSource11_11_10F>(RHICmdList, Format, UpdateType));
 				break;
 			case PF_A2B10G10R10:
-				RUN_TEST(Test_UpdateTexture2D_Impl<FDataSource8x4>(RHICmdList, Format));
+				RUN_TEST(Test_UpdateTexture_Impl<FDataSource8x4>(RHICmdList, Format, UpdateType));
 				break;
 			case PF_FloatRGBA:
-				RUN_TEST(Test_UpdateTexture2D_Impl<FDataSource16x4F>(RHICmdList, Format));
+				RUN_TEST(Test_UpdateTexture_Impl<FDataSource16x4F>(RHICmdList, Format, UpdateType));
 				break;
 			}
 		};
@@ -736,6 +812,22 @@ public:
 		return bResult;
 	}
 
+	static bool Test_UpdateTexture(FRHICommandListImmediate& RHICmdList)
+	{
+		bool bResult = true;
+
+		bResult &= Test_UpdateTexture_Impl(RHICmdList, ETextureUpdateType::Texture2D);
+
+		if (GRHISupportsUpdateFromBufferTexture)
+		{
+			bResult &= Test_UpdateTexture_Impl(RHICmdList, ETextureUpdateType::TextureFromBuffer);
+		}
+
+		bResult &= Test_UpdateTexture_Impl(RHICmdList, ETextureUpdateType::Texture3D);
+
+		return bResult;
+	}
+
 	//////////////////////////////////////////////////////////////////////////
 
 	static bool Test_RHICopyTexture(FRHICommandListImmediate& RHICmdList)
diff --git a/Engine/Source/Runtime/Apple/MetalRHI/Private/MetalDynamicRHI.h b/Engine/Source/Runtime/Apple/MetalRHI/Private/MetalDynamicRHI.h
--- a/Engine/Source/Runtime/Apple/MetalRHI/Private/MetalDynamicRHI.h
+++ b/Engine/Source/Runtime/Apple/MetalRHI/Private/MetalDynamicRHI.h
@@ -83,8 +83,8 @@ public:
 	virtual void RHIUnlockTexture2D(FRHITexture2D* Texture, uint32 MipIndex, bool bLockWithinMiptail) final override;
 	virtual void* RHILockTexture2DArray(FRHITexture2DArray* Texture, uint32 TextureIndex, uint32 MipIndex, EResourceLockMode LockMode, uint32& DestStride, bool bLockWithinMiptail) final override;
 	virtual void RHIUnlockTexture2DArray(FRHITexture2DArray* Texture, uint32 TextureIndex, uint32 MipIndex, bool bLockWithinMiptail) final override;
-	virtual void RHIUpdateTexture2D(FRHITexture2D* Texture, uint32 MipIndex, const struct FUpdateTextureRegion2D& UpdateRegion, uint32 SourcePitch, const uint8* SourceData) final override;
-	virtual void RHIUpdateTexture3D(FRHITexture3D* Texture, uint32 MipIndex, const struct FUpdateTextureRegion3D& UpdateRegion, uint32 SourceRowPitch, uint32 SourceDepthPitch, const uint8* SourceData) final override;
+	virtual void RHIUpdateTexture2D(FRHICommandListBase& RHICmdList, FRHITexture2D* Texture, uint32 MipIndex, const struct FUpdateTextureRegion2D& UpdateRegion, uint32 SourcePitch, const uint8* SourceData) final override;
+	virtual void RHIUpdateTexture3D(FRHICommandListBase& RHICmdList, FRHITexture3D* Texture, uint32 MipIndex, const struct FUpdateTextureRegion3D& UpdateRegion, uint32 SourceRowPitch, uint32 SourceDepthPitch, const uint8* SourceData) final override;
 	virtual void* RHILockTextureCubeFace(FRHITextureCube* Texture, uint32 FaceIndex, uint32 ArrayIndex, uint32 MipIndex, EResourceLockMode LockMode, uint32& DestStride, bool bLockWithinMiptail) final override;
 	virtual void RHIUnlockTextureCubeFace(FRHITextureCube* Texture, uint32 FaceIndex, uint32 ArrayIndex, uint32 MipIndex, bool bLockWithinMiptail) final override;
 	virtual void RHIBindDebugLabelName(FRHITexture* Texture, const TCHAR* Name) final override;
@@ -128,12 +128,10 @@ public:
 
 	virtual void* LockTexture2D_RenderThread(class FRHICommandListImmediate& RHICmdList, FRHITexture2D* Texture, uint32 MipIndex, EResourceLockMode LockMode, uint32& DestStride, bool bLockWithinMiptail, bool bNeedsDefaultRHIFlush = true) final override;
 	virtual void UnlockTexture2D_RenderThread(class FRHICommandListImmediate& RHICmdList, FRHITexture2D* Texture, uint32 MipIndex, bool bLockWithinMiptail, bool bNeedsDefaultRHIFlush = true) final override;
-	virtual void UpdateTexture2D_RenderThread(class FRHICommandListImmediate& RHICmdList, FRHITexture2D* Texture, uint32 MipIndex, const struct FUpdateTextureRegion2D& UpdateRegion, uint32 SourcePitch, const uint8* SourceData) final override;
-	
-	virtual FUpdateTexture3DData BeginUpdateTexture3D_RenderThread(class FRHICommandListImmediate& RHICmdList, FRHITexture3D* Texture, uint32 MipIndex, const struct FUpdateTextureRegion3D& UpdateRegion) final override;
-	virtual void EndUpdateTexture3D_RenderThread(class FRHICommandListImmediate& RHICmdList, FUpdateTexture3DData& UpdateData) final override;
-	virtual void UpdateTexture3D_RenderThread(class FRHICommandListImmediate& RHICmdList, FRHITexture3D* Texture, uint32 MipIndex, const struct FUpdateTextureRegion3D& UpdateRegion, uint32 SourceRowPitch, uint32 SourceDepthPitch, const uint8* SourceData) final override;
-	
+
+	virtual FUpdateTexture3DData RHIBeginUpdateTexture3D(FRHICommandListBase& RHICmdList, FRHITexture3D* Texture, uint32 MipIndex, const struct FUpdateTextureRegion3D& UpdateRegion) final override;
+	virtual void RHIEndUpdateTexture3D(FRHICommandListBase& RHICmdList, FUpdateTexture3DData& UpdateData) final override;
+
 	virtual FTextureRHIRef RHICreateTexture_RenderThread(class FRHICommandListImmediate& RHICmdList, const FRHITextureCreateDesc& CreateDesc) final override;
 	
 	virtual FRenderQueryRHIRef RHICreateRenderQuery_RenderThread(class FRHICommandListImmediate& RHICmdList, ERenderQueryType QueryType) final override;
diff --git a/Engine/Source/Runtime/Apple/MetalRHI/Private/MetalTexture.cpp b/Engine/Source/Runtime/Apple/MetalRHI/Private/MetalTexture.cpp
--- a/Engine/Source/Runtime/Apple/MetalRHI/Private/MetalTexture.cpp
+++ b/Engine/Source/Runtime/Apple/MetalRHI/Private/MetalTexture.cpp
@@ -1798,78 +1798,18 @@ static void InternalUpdateTexture2D(FMetalContext& Context, FRHITexture2D* Textu
 	FPlatformAtomics::InterlockedExchange(&Texture->Written, 1);
 }
 
-struct FMetalRHICommandUpdateTexture2D final : public FRHICommand<FMetalRHICommandUpdateTexture2D>
-{
-	FMetalContext& Context;
-	FRHITexture2D* Texture;
-	uint32 MipIndex;
-	FUpdateTextureRegion2D UpdateRegion;
-	uint32 SourcePitch;
-	FMetalBuffer SourceBuffer;
-
-	FORCEINLINE_DEBUGGABLE FMetalRHICommandUpdateTexture2D(FMetalContext& InContext, FRHITexture2D* InTexture, uint32 InMipIndex, FUpdateTextureRegion2D InUpdateRegion, uint32 InSourcePitch, const uint8* SourceData)
-	: Context(InContext)
-	, Texture(InTexture)
-	, MipIndex(InMipIndex)
-	, UpdateRegion(InUpdateRegion)
-	, SourcePitch(InSourcePitch)
-	{
-		SourceBuffer = Internal_CreateBufferAndCopyTexture2DUpdateRegionData(Texture, UpdateRegion, SourcePitch, SourceData);
-	}
-	
-	void Execute(FRHICommandListBase& CmdList)
-	{		
-		InternalUpdateTexture2D(Context, Texture, MipIndex, UpdateRegion, SourcePitch, SourceBuffer);
-		GetMetalDeviceContext().ReleaseBuffer(SourceBuffer);
-		INC_DWORD_STAT_BY(STAT_MetalTextureMemUpdate, UpdateRegion.Height * SourcePitch);
-	}
-};
-
-void FMetalDynamicRHI::UpdateTexture2D_RenderThread(class FRHICommandListImmediate& RHICmdList, FRHITexture2D* Texture, uint32 MipIndex, const struct FUpdateTextureRegion2D& UpdateRegion, uint32 SourcePitch, const uint8* SourceData)
-{
-	@autoreleasepool {
-		if (RHICmdList.Bypass() || !IsRunningRHIInSeparateThread())
-		{
-			this->RHIUpdateTexture2D(Texture, MipIndex, UpdateRegion, SourcePitch, SourceData);
-		}
-		else
-		{
-			new (RHICmdList.AllocCommand<FMetalRHICommandUpdateTexture2D>()) FMetalRHICommandUpdateTexture2D(ImmediateContext.GetInternalContext(), Texture, MipIndex, UpdateRegion, SourcePitch, SourceData);
-		}
-	}
-}
-
-void FMetalDynamicRHI::RHIUpdateTexture2D(FRHITexture2D* TextureRHI, uint32 MipIndex, const struct FUpdateTextureRegion2D& UpdateRegion, uint32 SourcePitch, const uint8* SourceData)
+void FMetalDynamicRHI::RHIUpdateTexture2D(FRHICommandListBase& RHICmdList, FRHITexture2D* TextureRHI, uint32 MipIndex, const struct FUpdateTextureRegion2D& UpdateRegion, uint32 SourcePitch, const uint8* SourceData)
 {
 	@autoreleasepool 
 	{
-		FMetalSurface* Texture = ResourceCast(TextureRHI);
-		FMetalTexture Tex = Texture->Texture;
-		
-		if(Tex.GetStorageMode() == mtlpp::StorageMode::Private)
+		FMetalBuffer Buffer = Internal_CreateBufferAndCopyTexture2DUpdateRegionData(TextureRHI, UpdateRegion, SourcePitch, SourceData);
+
+		RHICmdList.EnqueueLambda([TextureRHI, MipIndex, UpdateRegion, SourcePitch, Buffer](FRHICommandListBase& InRHICmdList) mutable
 		{
-			FMetalBuffer Buffer = Internal_CreateBufferAndCopyTexture2DUpdateRegionData(TextureRHI, UpdateRegion, SourcePitch, SourceData);
-			InternalUpdateTexture2D(ImmediateContext.GetInternalContext(), TextureRHI, MipIndex, UpdateRegion, SourcePitch, Buffer);
+			InternalUpdateTexture2D(static_cast<FMetalRHICommandContext&>(InRHICmdList.GetContext()).GetInternalContext(), TextureRHI, MipIndex, UpdateRegion, SourcePitch, Buffer);
 			GetMetalDeviceContext().ReleaseBuffer(Buffer);
-		}
-		else
-		{
-#if PLATFORM_MAC
-			// Expand R8_sRGB into RGBA8_sRGB for non Apple Silicon Mac.
-			TArray<uint32> ExpandedData;
-			if(Texture->GetFormat() == PF_G8 && EnumHasAnyFlags(Texture->GetFlags(), TexCreate_SRGB) && Tex.GetPixelFormat() == mtlpp::PixelFormat::RGBA8Unorm_sRGB)
-			{
-				ExpandedData.AddZeroed(UpdateRegion.Height * UpdateRegion.Width);
-				InternalExpandR8ToStandardRGBA((uint32*)ExpandedData.GetData(), UpdateRegion, SourcePitch, SourceData);
-				SourceData = (uint8*)ExpandedData.GetData();
-			}
-#endif
-			mtlpp::Region Region = mtlpp::Region(UpdateRegion.DestX, UpdateRegion.DestY, UpdateRegion.Width, UpdateRegion.Height);
-			MTLPP_VALIDATE(mtlpp::Texture, Tex, SafeGetRuntimeDebuggingLevel() >= EMetalDebugLevelValidation, Replace(Region, MipIndex, 0, SourceData, SourcePitch, 0));
+		});
 
-			FPlatformAtomics::InterlockedExchange(&Texture->Written, 1);
-		}
-		
 		INC_DWORD_STAT_BY(STAT_MetalTextureMemUpdate, UpdateRegion.Height*SourcePitch);
 	}
 }
@@ -1941,51 +1881,8 @@ static void InternalUpdateTexture3D(FMetalContext& Context, FRHITexture3D* Textu
 	FPlatformAtomics::InterlockedExchange(&Texture->Written, 1);
 }
 
-struct FMetalDynamicRHIUpdateTexture3DCommand final : public FRHICommand<FMetalDynamicRHIUpdateTexture3DCommand>
+FUpdateTexture3DData FMetalDynamicRHI::RHIBeginUpdateTexture3D(FRHICommandListBase& RHICmdList, FRHITexture3D* Texture, uint32 MipIndex, const struct FUpdateTextureRegion3D& UpdateRegion)
 {
-	FMetalContext& Context;
-	FRHITexture3D* DestinationTexture;
-	uint32 MipIndex;
-	FUpdateTextureRegion3D UpdateRegion;
-	uint32 SourceRowPitch;
-	uint32 SourceDepthPitch;
-	FMetalBuffer Buffer;
-	
-	FORCEINLINE_DEBUGGABLE FMetalDynamicRHIUpdateTexture3DCommand(FMetalContext& InContext, FRHITexture3D* TextureRHI, uint32 InMipIndex, const struct FUpdateTextureRegion3D& InUpdateRegion, uint32 InSourceRowPitch, uint32 InSourceDepthPitch, const uint8* SourceData)
-	: Context(InContext)
-	, DestinationTexture(TextureRHI)
-	, MipIndex(InMipIndex)
-	, UpdateRegion(InUpdateRegion)
-	, SourceRowPitch(InSourceRowPitch)
-	, SourceDepthPitch(InSourceDepthPitch)
-	{
-		Buffer = Internal_CreateBufferAndCopyTexture3DUpdateRegionData(TextureRHI, UpdateRegion, SourceRowPitch, SourceDepthPitch, SourceData);
-	}
-	
-	void Execute(FRHICommandListBase& CmdList)
-	{
-		InternalUpdateTexture3D(Context, DestinationTexture, MipIndex, UpdateRegion, SourceRowPitch, SourceDepthPitch, Buffer);
-		GetMetalDeviceContext().ReleaseBuffer(Buffer);
-		INC_DWORD_STAT_BY(STAT_MetalTextureMemUpdate, UpdateRegion.Height * UpdateRegion.Width * SourceDepthPitch);
-	}
-};
-
-void FMetalDynamicRHI::UpdateTexture3D_RenderThread(class FRHICommandListImmediate& RHICmdList, FRHITexture3D* Texture, uint32 MipIndex, const struct FUpdateTextureRegion3D& UpdateRegion, uint32 SourceRowPitch, uint32 SourceDepthPitch, const uint8* SourceData)
-{
-	if (RHICmdList.Bypass() || !IsRunningRHIInSeparateThread())
-	{
-		RHIUpdateTexture3D(Texture, MipIndex, UpdateRegion, SourceRowPitch, SourceDepthPitch, SourceData);
-	}
-	else
-	{
-		new (RHICmdList.AllocCommand<FMetalDynamicRHIUpdateTexture3DCommand>()) FMetalDynamicRHIUpdateTexture3DCommand(ImmediateContext.GetInternalContext(), Texture, MipIndex, UpdateRegion, SourceRowPitch, SourceDepthPitch, SourceData);
-	}
-}
-
-FUpdateTexture3DData FMetalDynamicRHI::BeginUpdateTexture3D_RenderThread(class FRHICommandListImmediate& RHICmdList, FRHITexture3D* Texture, uint32 MipIndex, const struct FUpdateTextureRegion3D& UpdateRegion)
-{
-	check(IsInRenderingThread());
-	
 	const int32 FormatSize = PixelFormatBlockBytes[Texture->GetFormat()];
 	const int32 RowPitch = UpdateRegion.Width * FormatSize;
 	const int32 DepthPitch = UpdateRegion.Width * UpdateRegion.Height * FormatSize;
@@ -1996,49 +1893,27 @@ FUpdateTexture3DData FMetalDynamicRHI::BeginUpdateTexture3D_RenderThread(class F
 	return FUpdateTexture3DData(Texture, MipIndex, UpdateRegion, RowPitch, DepthPitch, Data, MemorySize, GFrameNumberRenderThread);
 }
 
-void FMetalDynamicRHI::EndUpdateTexture3D_RenderThread(class FRHICommandListImmediate& RHICmdList, FUpdateTexture3DData& UpdateData)
+void FMetalDynamicRHI::RHIEndUpdateTexture3D(FRHICommandListBase& RHICmdList, FUpdateTexture3DData& UpdateData)
 {
-	check(IsInRenderingThread());
+	check(IsInParallelRenderingThread());
 	check(GFrameNumberRenderThread == UpdateData.FrameNumber);
-	
-	if (RHICmdList.Bypass() || !IsRunningRHIInSeparateThread())
-	{
-		GDynamicRHI->RHIUpdateTexture3D(UpdateData.Texture, UpdateData.MipIndex, UpdateData.UpdateRegion, UpdateData.RowPitch, UpdateData.DepthPitch, UpdateData.Data);
-	}
-	else
-	{
-		new (RHICmdList.AllocCommand<FMetalDynamicRHIUpdateTexture3DCommand>()) FMetalDynamicRHIUpdateTexture3DCommand(ImmediateContext.GetInternalContext(), UpdateData.Texture, UpdateData.MipIndex, UpdateData.UpdateRegion, UpdateData.RowPitch, UpdateData.DepthPitch, UpdateData.Data);
-	}
-	
+	GDynamicRHI->RHIUpdateTexture3D(RHICmdList, UpdateData.Texture, UpdateData.MipIndex, UpdateData.UpdateRegion, UpdateData.RowPitch, UpdateData.DepthPitch, UpdateData.Data);
 	FMemory::Free(UpdateData.Data);
 	UpdateData.Data = nullptr;
 }
 
-void FMetalDynamicRHI::RHIUpdateTexture3D(FRHITexture3D* TextureRHI,uint32 MipIndex,const FUpdateTextureRegion3D& UpdateRegion,uint32 SourceRowPitch,uint32 SourceDepthPitch, const uint8* SourceData)
+void FMetalDynamicRHI::RHIUpdateTexture3D(FRHICommandListBase& RHICmdList, FRHITexture3D* TextureRHI,uint32 MipIndex,const FUpdateTextureRegion3D& UpdateRegion,uint32 SourceRowPitch,uint32 SourceDepthPitch, const uint8* SourceData)
 {
-	@autoreleasepool {
-	
-		FMetalSurface* Texture = ResourceCast(TextureRHI);
-		FMetalTexture Tex = Texture->Texture;
-		
-#if PLATFORM_MAC
-		checkf(!(Texture->GetFormat() == PF_G8 && EnumHasAnyFlags(Texture->GetFlags(), TexCreate_SRGB) && Tex.GetPixelFormat() == mtlpp::PixelFormat::RGBA8Unorm_sRGB), TEXT("MetalRHI on non Apple Silicon does not support PF_G8_sRGB on 3D, array or cube textures as it requires manual, CPU-side expansion to RGBA8_sRGB which is expensive!"));
-#endif
-		if(Tex.GetStorageMode() == mtlpp::StorageMode::Private)
-		{
-			SCOPED_AUTORELEASE_POOL;
+	@autoreleasepool
+	{
+		FMetalBuffer Buffer = Internal_CreateBufferAndCopyTexture3DUpdateRegionData(TextureRHI, UpdateRegion, SourceRowPitch, SourceDepthPitch, SourceData);
 
-			FMetalBuffer IntermediateBuffer = Internal_CreateBufferAndCopyTexture3DUpdateRegionData(TextureRHI, UpdateRegion, SourceRowPitch, SourceDepthPitch, SourceData);
-			InternalUpdateTexture3D(ImmediateContext.GetInternalContext(), TextureRHI, MipIndex, UpdateRegion, SourceRowPitch, SourceDepthPitch, IntermediateBuffer);
-			GetMetalDeviceContext().ReleaseBuffer(IntermediateBuffer);
-		}
-		else
+		RHICmdList.EnqueueLambda([TextureRHI, MipIndex, UpdateRegion, SourceRowPitch, SourceDepthPitch, Buffer](FRHICommandListBase& InRHICmdList) mutable
 		{
-			mtlpp::Region Region = mtlpp::Region(UpdateRegion.DestX, UpdateRegion.DestY, UpdateRegion.DestZ, UpdateRegion.Width, UpdateRegion.Height, UpdateRegion.Depth);
-			MTLPP_VALIDATE(mtlpp::Texture, Tex, SafeGetRuntimeDebuggingLevel() >= EMetalDebugLevelValidation, Replace(Region, MipIndex, 0, SourceData, SourceRowPitch, SourceDepthPitch));
-			FPlatformAtomics::InterlockedExchange(&Texture->Written, 1);
-		}
-		
+			InternalUpdateTexture3D(static_cast<FMetalRHICommandContext&>(InRHICmdList.GetContext()).GetInternalContext(), TextureRHI, MipIndex, UpdateRegion, SourceRowPitch, SourceDepthPitch, Buffer);
+			GetMetalDeviceContext().ReleaseBuffer(Buffer);
+		});
+
 		INC_DWORD_STAT_BY(STAT_MetalTextureMemUpdate, UpdateRegion.Height * UpdateRegion.Width * SourceDepthPitch);
 	}
 }
diff --git a/Engine/Source/Runtime/D3D12RHI/Private/D3D12Buffer.cpp b/Engine/Source/Runtime/D3D12RHI/Private/D3D12Buffer.cpp
--- a/Engine/Source/Runtime/D3D12RHI/Private/D3D12Buffer.cpp
+++ b/Engine/Source/Runtime/D3D12RHI/Private/D3D12Buffer.cpp
@@ -230,15 +230,7 @@ void FD3D12Buffer::UploadResourceData(FRHICommandListBase& RHICmdList, FResource
 		}
 		else
 		{
-			// @todo d3d12 - Refactor contexts so pipe switching is not required.
-			TOptional<ERHIPipeline> PreviousPipeline;
-			if (RHICmdList.GetPipeline() == ERHIPipeline::None)
-				PreviousPipeline = RHICmdList.SwitchPipeline(ERHIPipeline::Graphics);
-
 			new (RHICmdList.AllocCommand<FD3D12RHICommandInitializeBuffer>()) FD3D12RHICommandInitializeBuffer(this, SrcResourceLoc, BufferSize, InDestinationState);
-
-			if (PreviousPipeline.IsSet())
-				RHICmdList.SwitchPipeline(PreviousPipeline.GetValue());
 		}
 	}
 
diff --git a/Engine/Source/Runtime/D3D12RHI/Private/D3D12RHIPrivate.h b/Engine/Source/Runtime/D3D12RHI/Private/D3D12RHIPrivate.h
--- a/Engine/Source/Runtime/D3D12RHI/Private/D3D12RHIPrivate.h
+++ b/Engine/Source/Runtime/D3D12RHI/Private/D3D12RHIPrivate.h
@@ -476,11 +476,11 @@ public:
 	virtual void RHIUnlockTexture2D(FRHITexture2D* Texture, uint32 MipIndex, bool bLockWithinMiptail) final override;
 	virtual void* RHILockTexture2DArray(FRHITexture2DArray* Texture, uint32 TextureIndex, uint32 MipIndex, EResourceLockMode LockMode, uint32& DestStride, bool bLockWithinMiptail) final override;
 	virtual void RHIUnlockTexture2DArray(FRHITexture2DArray* Texture, uint32 TextureIndex, uint32 MipIndex, bool bLockWithinMiptail) final override;
-	virtual void RHIUpdateTexture2D(FRHITexture2D* Texture, uint32 MipIndex, const struct FUpdateTextureRegion2D& UpdateRegion, uint32 SourcePitch, const uint8* SourceData) final override;
-	virtual void RHIUpdateTexture3D(FRHITexture3D* Texture, uint32 MipIndex, const struct FUpdateTextureRegion3D& UpdateRegion, uint32 SourceRowPitch, uint32 SourceDepthPitch, const uint8* SourceData) final override;
-	virtual FUpdateTexture3DData BeginUpdateTexture3D_RenderThread(class FRHICommandListImmediate& RHICmdList, FRHITexture3D* Texture, uint32 MipIndex, const struct FUpdateTextureRegion3D& UpdateRegion) final override;
-	virtual void EndUpdateTexture3D_RenderThread(class FRHICommandListImmediate& RHICmdList, FUpdateTexture3DData& UpdateData) final override;
-	virtual void EndMultiUpdateTexture3D_RenderThread(class FRHICommandListImmediate& RHICmdList, TArray<FUpdateTexture3DData>& UpdateDataArray) final override;
+	virtual void RHIUpdateTexture2D(FRHICommandListBase& RHICmdList, FRHITexture2D* Texture, uint32 MipIndex, const struct FUpdateTextureRegion2D& UpdateRegion, uint32 SourcePitch, const uint8* SourceData) final override;
+	virtual void RHIUpdateTexture3D(FRHICommandListBase& RHICmdList, FRHITexture3D* Texture, uint32 MipIndex, const struct FUpdateTextureRegion3D& UpdateRegion, uint32 SourceRowPitch, uint32 SourceDepthPitch, const uint8* SourceData) final override;
+	virtual FUpdateTexture3DData RHIBeginUpdateTexture3D(FRHICommandListBase& RHICmdList, FRHITexture3D* Texture, uint32 MipIndex, const struct FUpdateTextureRegion3D& UpdateRegion) final override;
+	virtual void RHIEndUpdateTexture3D(FRHICommandListBase& RHICmdList, FUpdateTexture3DData& UpdateData) final override;
+	virtual void RHIEndMultiUpdateTexture3D(FRHICommandListBase& RHICmdList, TArray<FUpdateTexture3DData>& UpdateDataArray) final override;
 	virtual void* RHILockTextureCubeFace(FRHITextureCube* Texture, uint32 FaceIndex, uint32 ArrayIndex, uint32 MipIndex, EResourceLockMode LockMode, uint32& DestStride, bool bLockWithinMiptail) final override;
 	virtual void RHIUnlockTextureCubeFace(FRHITextureCube* Texture, uint32 FaceIndex, uint32 ArrayIndex, uint32 MipIndex, bool bLockWithinMiptail) final override;
 	virtual void RHIBindDebugLabelName(FRHITexture* Texture, const TCHAR* Name) final override;
@@ -601,7 +601,6 @@ public:
 		return RHICreateComputeShader(Code, Hash);
 	}
 
-	virtual void UpdateTexture2D_RenderThread(class FRHICommandListImmediate& RHICmdList, FRHITexture2D* Texture, uint32 MipIndex, const struct FUpdateTextureRegion2D& UpdateRegion, uint32 SourcePitch, const uint8* SourceData) override final;
 	virtual void* LockTexture2D_RenderThread(class FRHICommandListImmediate& RHICmdList, FRHITexture2D* Texture, uint32 MipIndex, EResourceLockMode LockMode, uint32& DestStride, bool bLockWithinMiptail, bool bNeedsDefaultRHIFlush = true) override final;
 	virtual void UnlockTexture2D_RenderThread(class FRHICommandListImmediate& RHICmdList, FRHITexture2D* Texture, uint32 MipIndex, bool bLockWithinMiptail, bool bNeedsDefaultRHIFlush = true) override final;
 
@@ -687,13 +686,13 @@ public:
 		// Not supported on PC
 		return false;
 	}
-	virtual void EndUpdateTexture3D_ComputeShader(FUpdateTexture3DData& UpdateData, FD3D12UpdateTexture3DData* UpdateDataD3D12)
+	virtual void EndUpdateTexture3D_ComputeShader(FRHICommandListBase& RHICmdList, FUpdateTexture3DData& UpdateData, FD3D12UpdateTexture3DData* UpdateDataD3D12)
 	{
 		// Not supported on PC
 	}
 
 	FUpdateTexture3DData BeginUpdateTexture3D_Internal(FRHITexture3D* Texture, uint32 MipIndex, const struct FUpdateTextureRegion3D& UpdateRegion);
-	void EndUpdateTexture3D_Internal(FUpdateTexture3DData& UpdateData);
+	void EndUpdateTexture3D_Internal(FRHICommandListBase& RHICmdList, FUpdateTexture3DData& UpdateData);
 
 	void UpdateBuffer(FD3D12ResourceLocation* Dest, uint32 DestOffset, FD3D12ResourceLocation* Source, uint32 SourceOffset, uint32 NumBytes);
 
diff --git a/Engine/Source/Runtime/D3D12RHI/Private/D3D12Texture.cpp b/Engine/Source/Runtime/D3D12RHI/Private/D3D12Texture.cpp
--- a/Engine/Source/Runtime/D3D12RHI/Private/D3D12Texture.cpp
+++ b/Engine/Source/Runtime/D3D12RHI/Private/D3D12Texture.cpp
@@ -2013,7 +2013,7 @@ void FD3D12Texture::UnlockInternal(class FRHICommandListImmediate* RHICmdList, F
 	}
 }
 
-void FD3D12Texture::UpdateTexture2D(class FRHICommandListImmediate* RHICmdList, uint32 MipIndex, const FUpdateTextureRegion2D& UpdateRegion, uint32 SourcePitch, const uint8* SourceData)
+void FD3D12Texture::UpdateTexture2D(FRHICommandListBase& RHICmdList, uint32 MipIndex, const FUpdateTextureRegion2D& UpdateRegion, uint32 SourcePitch, const uint8* SourceData)
 {
 	const FPixelFormatInfo& FormatInfo = GPixelFormats[this->GetFormat()];
 	check(UpdateRegion.Width  %	FormatInfo.BlockSizeX == 0);
@@ -2061,9 +2061,9 @@ void FD3D12Texture::UpdateTexture2D(class FRHICommandListImmediate* RHICmdList,
 		CD3DX12_TEXTURE_COPY_LOCATION SourceCopyLocation(UploadHeapResourceLocation.GetResource()->GetResource(), PlacedTexture2D);
 
 		// If we are on the render thread, queue up the copy on the RHIThread so it happens at the correct time.
-		if (ShouldDeferCmdListOperation(RHICmdList))
+		if (RHICmdList.IsTopOfPipe())
 		{
-			ALLOC_COMMAND_CL(*RHICmdList, FRHICommandUpdateTexture)(&Texture, MipIndex, UpdateRegion.DestX, UpdateRegion.DestY, 0, SourceCopyLocation, &UploadHeapResourceLocation);
+			ALLOC_COMMAND_CL(RHICmdList, FRHICommandUpdateTexture)(&Texture, MipIndex, UpdateRegion.DestX, UpdateRegion.DestY, 0, SourceCopyLocation, &UploadHeapResourceLocation);
 		}
 		else
 		{
@@ -2186,36 +2186,21 @@ void FD3D12DynamicRHI::RHIUnlockTexture2DArray(FRHITexture2DArray* TextureRHI, u
 	Texture->Unlock(nullptr, MipIndex, TextureIndex);
 }
 
-void FD3D12DynamicRHI::UpdateTexture2D_RenderThread(class FRHICommandListImmediate& RHICmdList, FRHITexture2D* TextureRHI, uint32 MipIndex, const struct FUpdateTextureRegion2D& UpdateRegion, uint32 SourcePitch, const uint8* SourceData)
+void FD3D12DynamicRHI::RHIUpdateTexture2D(FRHICommandListBase& RHICmdList, FRHITexture2D* TextureRHI, uint32 MipIndex, const FUpdateTextureRegion2D& UpdateRegion, uint32 SourcePitch, const uint8* SourceData)
 {
 	check(TextureRHI);
 	FD3D12Texture* Texture = FD3D12DynamicRHI::ResourceCast(TextureRHI);
-	Texture->UpdateTexture2D(&RHICmdList, MipIndex, UpdateRegion, SourcePitch, SourceData);
+	Texture->UpdateTexture2D(RHICmdList, MipIndex, UpdateRegion, SourcePitch, SourceData);
 }
 
-void FD3D12DynamicRHI::RHIUpdateTexture2D(FRHITexture2D* TextureRHI, uint32 MipIndex, const FUpdateTextureRegion2D& UpdateRegion, uint32 SourcePitch, const uint8* SourceData)
+FUpdateTexture3DData FD3D12DynamicRHI::RHIBeginUpdateTexture3D(FRHICommandListBase& RHICmdList, FRHITexture3D* Texture, uint32 MipIndex, const struct FUpdateTextureRegion3D& UpdateRegion)
 {
-	check(TextureRHI);
-	FD3D12Texture* Texture = FD3D12DynamicRHI::ResourceCast(TextureRHI);
-	Texture->UpdateTexture2D(nullptr, MipIndex, UpdateRegion, SourcePitch, SourceData);
-}
-
-FUpdateTexture3DData FD3D12DynamicRHI::BeginUpdateTexture3D_RenderThread(class FRHICommandListImmediate& RHICmdList, FRHITexture3D* Texture, uint32 MipIndex, const struct FUpdateTextureRegion3D& UpdateRegion)
-{
-	check(IsInRenderingThread());
-	// This stall could potentially be removed, provided the fast allocator is thread-safe. However we 
-	// currently need to stall in the End method anyway (see below)
-	RHICmdList.ImmediateFlush(EImmediateFlushType::FlushRHIThread);
 	return BeginUpdateTexture3D_Internal(Texture, MipIndex, UpdateRegion);
 }
 
-void FD3D12DynamicRHI::EndUpdateTexture3D_RenderThread(class FRHICommandListImmediate& RHICmdList, FUpdateTexture3DData& UpdateData)
+void FD3D12DynamicRHI::RHIEndUpdateTexture3D(FRHICommandListBase& RHICmdList, FUpdateTexture3DData& UpdateData)
 {
-	check(IsInRenderingThread());
-	// TODO: move this command entirely to the RHI thread so we can remove these stalls
-	// and fix potential ordering issue with non-compute-shader version
-	RHICmdList.ImmediateFlush(EImmediateFlushType::FlushRHIThread);
-	EndUpdateTexture3D_Internal(UpdateData);
+	EndUpdateTexture3D_Internal(RHICmdList, UpdateData);
 }
 
 struct FD3D12RHICmdEndMultiUpdateTexture3DString
@@ -2335,9 +2320,9 @@ private:
 };
 
 // Single pair of transition barriers instead of one pair for each update
-void FD3D12DynamicRHI::EndMultiUpdateTexture3D_RenderThread(class FRHICommandListImmediate& RHICmdList, TArray<FUpdateTexture3DData>& UpdateDataArray)
+void FD3D12DynamicRHI::RHIEndMultiUpdateTexture3D(FRHICommandListBase& RHICmdList, TArray<FUpdateTexture3DData>& UpdateDataArray)
 {
-	check(IsInRenderingThread());
+	check(IsInParallelRenderingThread());
 	check(UpdateDataArray.Num() > 0);
 	check(GFrameNumberRenderThread == UpdateDataArray[0].FrameNumber);
 #if DO_CHECK
@@ -2364,12 +2349,12 @@ void FD3D12DynamicRHI::EndMultiUpdateTexture3D_RenderThread(class FRHICommandLis
 			FUpdateTexture3DData& UpdateData = UpdateDataArray[Idx];
 			FD3D12UpdateTexture3DData* UpdateDataD3D12 =
 				reinterpret_cast<FD3D12UpdateTexture3DData*>(&UpdateData.PlatformData[0]);
-			EndUpdateTexture3D_ComputeShader(UpdateData, UpdateDataD3D12);
+			EndUpdateTexture3D_ComputeShader(RHICmdList, UpdateData, UpdateDataD3D12);
 		}
 	}
 	else
 	{
-		if (RHICmdList.Bypass())
+		if (RHICmdList.IsBottomOfPipe())
 		{
 			FD3D12RHICmdEndMultiUpdateTexture3D RHICmd(UpdateDataArray);
 			RHICmd.Execute(RHICmdList);
@@ -2381,24 +2366,29 @@ void FD3D12DynamicRHI::EndMultiUpdateTexture3D_RenderThread(class FRHICommandLis
 	}
 }
 
-void FD3D12DynamicRHI::RHIUpdateTexture3D(FRHITexture3D* TextureRHI, uint32 MipIndex, const FUpdateTextureRegion3D& InUpdateRegion, uint32 SourceRowPitch, uint32 SourceDepthPitch, const uint8* SourceData)
+void FD3D12DynamicRHI::RHIUpdateTexture3D(FRHICommandListBase& RHICmdList, FRHITexture3D* TextureRHI, uint32 MipIndex, const FUpdateTextureRegion3D& InUpdateRegion, uint32 SourceRowPitch, uint32 SourceDepthPitch, const uint8* SourceData)
 {
-	check(IsInRenderingThread());
-
 	FD3D12Texture* Texture = FD3D12DynamicRHI::ResourceCast(TextureRHI);
 	const FPixelFormatInfo& FormatInfo = GPixelFormats[Texture->GetFormat()];
 
 	// Need to round up the height and with by block size.
 	FUpdateTextureRegion3D UpdateRegion = InUpdateRegion;
-	UpdateRegion.Width = FMath::DivideAndRoundUp<int32>(UpdateRegion.Width, FormatInfo.BlockSizeX) * FormatInfo.BlockSizeX;
-	UpdateRegion.Height = FMath::DivideAndRoundUp<int32>(UpdateRegion.Height, FormatInfo.BlockSizeY) * FormatInfo.BlockSizeY;
+
+	const uint32 NumBlockX = (uint32)FMath::DivideAndRoundUp<int32>(UpdateRegion.Width, FormatInfo.BlockSizeX);
+	const uint32 NumBlockY = (uint32)FMath::DivideAndRoundUp<int32>(UpdateRegion.Height, FormatInfo.BlockSizeY);
+
+	UpdateRegion.Width  = NumBlockX * FormatInfo.BlockSizeX;
+	UpdateRegion.Height = NumBlockY * FormatInfo.BlockSizeY;
 
 	FUpdateTexture3DData UpdateData = BeginUpdateTexture3D_Internal(TextureRHI, MipIndex, UpdateRegion);
 
+	const uint32 UpdateBytesRow = NumBlockX * FormatInfo.BlockBytes;
+	const uint32 UpdateBytesDepth = NumBlockY * UpdateBytesRow;
+
 	// Copy the data into the UpdateData destination buffer
 	check(nullptr != UpdateData.Data);
-	check(SourceRowPitch <= UpdateData.RowPitch);
-	check(SourceDepthPitch <= UpdateData.DepthPitch);
+	check(SourceRowPitch >= UpdateBytesRow);
+	check(SourceDepthPitch >= UpdateBytesDepth);
 
 	const uint32 NumRows = UpdateRegion.Height / (uint32)FormatInfo.BlockSizeY;
 
@@ -2409,19 +2399,19 @@ void FD3D12DynamicRHI::RHIUpdateTexture3D(FRHITexture3D* TextureRHI, uint32 MipI
 
 		for (uint32 j = 0; j < NumRows; j++)
 		{
-			FMemory::Memcpy(DestRowData, SourceRowData, SourceRowPitch);
+			FMemory::Memcpy(DestRowData, SourceRowData, UpdateBytesRow);
 			SourceRowData += SourceRowPitch;
 			DestRowData += UpdateData.RowPitch;
 		}
 	}
 
-	EndUpdateTexture3D_Internal(UpdateData);
+	EndUpdateTexture3D_Internal(RHICmdList, UpdateData);
 }
 
 
 FUpdateTexture3DData FD3D12DynamicRHI::BeginUpdateTexture3D_Internal(FRHITexture3D* TextureRHI, uint32 MipIndex, const struct FUpdateTextureRegion3D& UpdateRegion)
 {
-	check(IsInRenderingThread());
+	check(IsInParallelRenderingThread());
 	FUpdateTexture3DData UpdateData(TextureRHI, MipIndex, UpdateRegion, 0, 0, nullptr, 0, GFrameNumberRenderThread);
 		
 	// Initialize the platform data
@@ -2560,9 +2550,8 @@ private:
 	D3D12_PLACED_SUBRESOURCE_FOOTPRINT PlacedSubresourceFootprint;
 };
 
-void FD3D12DynamicRHI::EndUpdateTexture3D_Internal(FUpdateTexture3DData& UpdateData)
+void FD3D12DynamicRHI::EndUpdateTexture3D_Internal(FRHICommandListBase& RHICmdList, FUpdateTexture3DData& UpdateData)
 {
-	check(IsInRenderingThread());
 	check(GFrameNumberRenderThread == UpdateData.FrameNumber);
 
 	FD3D12UpdateTexture3DData* UpdateDataD3D12 = reinterpret_cast<FD3D12UpdateTexture3DData*>(&UpdateData.PlatformData[0]);
@@ -2570,12 +2559,11 @@ void FD3D12DynamicRHI::EndUpdateTexture3D_Internal(FUpdateTexture3DData& UpdateD
 
 	if (UpdateDataD3D12->bComputeShaderCopy)
 	{
-		EndUpdateTexture3D_ComputeShader(UpdateData, UpdateDataD3D12);
+		EndUpdateTexture3D_ComputeShader(RHICmdList, UpdateData, UpdateDataD3D12);
 	}
 	else
 	{
-		FRHICommandListImmediate& RHICmdList = FRHICommandListExecutor::GetImmediateCommandList();
-		if (RHICmdList.Bypass())
+		if (RHICmdList.IsBottomOfPipe())
 		{
 			FD3D12RHICmdEndUpdateTexture3D RHICmd(UpdateData);
 			RHICmd.Execute(RHICmdList);
diff --git a/Engine/Source/Runtime/D3D12RHI/Private/D3D12Texture.h b/Engine/Source/Runtime/D3D12RHI/Private/D3D12Texture.h
--- a/Engine/Source/Runtime/D3D12RHI/Private/D3D12Texture.h
+++ b/Engine/Source/Runtime/D3D12RHI/Private/D3D12Texture.h
@@ -95,7 +95,7 @@ public:
 	// Locking/update functions
 	void* Lock(class FRHICommandListImmediate* RHICmdList, uint32 MipIndex, uint32 ArrayIndex, EResourceLockMode LockMode, uint32& DestStride);
 	void Unlock(class FRHICommandListImmediate* RHICmdList, uint32 MipIndex, uint32 ArrayIndex);
-	void UpdateTexture2D(class FRHICommandListImmediate* RHICmdList, uint32 MipIndex, const struct FUpdateTextureRegion2D& UpdateRegion, uint32 SourcePitch, const uint8* SourceData);
+	void UpdateTexture2D(FRHICommandListBase& RHICmdList, uint32 MipIndex, const struct FUpdateTextureRegion2D& UpdateRegion, uint32 SourcePitch, const uint8* SourceData);
 	void UpdateTexture(uint32 MipIndex, uint32 DestX, uint32 DestY, uint32 DestZ, const D3D12_TEXTURE_COPY_LOCATION& SourceCopyLocation);
 	void CopyTextureRegion(uint32 DestX, uint32 DestY, uint32 DestZ, FD3D12Texture* SourceTexture, const D3D12_BOX& SourceBox);
 
diff --git a/Engine/Source/Runtime/Engine/Private/Materials/MaterialShared.cpp b/Engine/Source/Runtime/Engine/Private/Materials/MaterialShared.cpp
--- a/Engine/Source/Runtime/Engine/Private/Materials/MaterialShared.cpp
+++ b/Engine/Source/Runtime/Engine/Private/Materials/MaterialShared.cpp
@@ -3914,6 +3914,7 @@ public:
 			TRACE_CPUPROFILER_EVENT_SCOPE(FUniformExpressionCacheAsyncUpdater::Update);
 			FTaskTagScope Scope(ETaskTag::EParallelRenderingThread);
 			FMemMark Mark(FMemStack::Get());
+			RHICmdList->SwitchPipeline(ERHIPipeline::Graphics);
 
 			for (const FItem& Item : Items)
 			{
diff --git a/Engine/Source/Runtime/NullDrv/Public/NullRHI.h b/Engine/Source/Runtime/NullDrv/Public/NullRHI.h
--- a/Engine/Source/Runtime/NullDrv/Public/NullRHI.h
+++ b/Engine/Source/Runtime/NullDrv/Public/NullRHI.h
@@ -281,11 +281,11 @@ public:
 	{
 
 	}
-	virtual void RHIUpdateTexture2D(FRHITexture2D* Texture, uint32 MipIndex, const struct FUpdateTextureRegion2D& UpdateRegion, uint32 SourcePitch, const uint8* SourceData) final override
+	virtual void RHIUpdateTexture2D(FRHICommandListBase&, FRHITexture2D* Texture, uint32 MipIndex, const struct FUpdateTextureRegion2D& UpdateRegion, uint32 SourcePitch, const uint8* SourceData) final override
 	{
 
 	}
-	virtual void RHIUpdateTexture3D(FRHITexture3D* Texture, uint32 MipIndex, const struct FUpdateTextureRegion3D& UpdateRegion, uint32 SourceRowPitch, uint32 SourceDepthPitch, const uint8* SourceData) final override
+	virtual void RHIUpdateTexture3D(FRHICommandListBase&, FRHITexture3D* Texture, uint32 MipIndex, const struct FUpdateTextureRegion3D& UpdateRegion, uint32 SourceRowPitch, uint32 SourceDepthPitch, const uint8* SourceData) final override
 	{
 
 	}
diff --git a/Engine/Source/Runtime/OpenGLDrv/Private/OpenGLTexture.cpp b/Engine/Source/Runtime/OpenGLDrv/Private/OpenGLTexture.cpp
--- a/Engine/Source/Runtime/OpenGLDrv/Private/OpenGLTexture.cpp
+++ b/Engine/Source/Runtime/OpenGLDrv/Private/OpenGLTexture.cpp
@@ -1832,9 +1832,8 @@ void FOpenGLDynamicRHI::RHIUnlockTexture2DArray(FRHITexture2DArray* TextureRHI,
 	GetOpenGLTextureFromRHITexture(TextureRHI)->Unlock(MipIndex, TextureIndex);
 }
 
-void FOpenGLDynamicRHI::RHIUpdateTexture2D(FRHITexture2D* TextureRHI, uint32 MipIndex, const FUpdateTextureRegion2D& UpdateRegionIn, uint32 SourcePitch, const uint8* SourceDataIn)
+void FOpenGLDynamicRHI::RHIUpdateTexture2D(FRHICommandListBase& RHICmdList, FRHITexture2D* TextureRHI, uint32 MipIndex, const FUpdateTextureRegion2D& UpdateRegionIn, uint32 SourcePitch, const uint8* SourceDataIn)
 {
-	FRHICommandListImmediate& RHICmdList = FRHICommandListExecutor::GetImmediateCommandList();
 	const FUpdateTextureRegion2D UpdateRegion = UpdateRegionIn;
 
 	uint8* RHITSourceData = nullptr;
@@ -1847,7 +1846,7 @@ void FOpenGLDynamicRHI::RHIUpdateTexture2D(FRHITexture2D* TextureRHI, uint32 Mip
 		FMemory::Memcpy(RHITSourceData, SourceDataIn, DataSize);
 	}
 	const uint8* SourceData = RHITSourceData ? RHITSourceData : SourceDataIn;
-	RunOnGLRenderContextThread([=]()
+	RHICmdList.EnqueueLambda([this, TextureRHI, MipIndex, UpdateRegion, SourcePitch, SourceData, RHITSourceData](FRHICommandListBase&)
 	{
 		VERIFY_GL_SCOPE();
 
@@ -1886,59 +1885,76 @@ void FOpenGLDynamicRHI::RHIUpdateTexture2D(FRHITexture2D* TextureRHI, uint32 Mip
 	});
 }
 
-void FOpenGLDynamicRHI::RHIUpdateTexture3D(FRHITexture3D* TextureRHI, uint32 MipIndex, const FUpdateTextureRegion3D& UpdateRegion, uint32 SourceRowPitch, uint32 SourceDepthPitch, const uint8* SourceData)
+void FOpenGLDynamicRHI::RHIUpdateTexture3D(FRHICommandListBase& RHICmdList, FRHITexture3D* TextureRHI, uint32 MipIndex, const FUpdateTextureRegion3D& UpdateRegion, uint32 SourceRowPitch, uint32 SourceDepthPitch, const uint8* SourceData)
 {
-	VERIFY_GL_SCOPE();
-	check( FOpenGL::SupportsTexture3D() );
-	FOpenGLTexture* Texture = GetOpenGLTextureFromRHITexture(TextureRHI);
+	uint8* RHITSourceData = nullptr;
+	if (!ShouldRunGLRenderContextOpOnThisThread(RHICmdList))
+	{
+		uint32 DataSize = SourceDepthPitch * UpdateRegion.Depth;
+		RHITSourceData = (uint8*)FMemory::Malloc(DataSize, 16);
+		FMemory::Memcpy(RHITSourceData, SourceData, DataSize);
+		SourceData = RHITSourceData;
+	}
 
-	// Use a texture stage that's not likely to be used for draws, to avoid waiting
-	FOpenGLContextState& ContextState = GetContextStateForCurrentContext();
-	CachedSetupTextureStage(ContextState, FOpenGL::GetMaxCombinedTextureImageUnits() - 1, Texture->Target, Texture->GetResource(), 0, Texture->GetNumMips());
-	CachedBindPixelUnpackBuffer(ContextState, 0);
+	RHICmdList.EnqueueLambda([this, TextureRHI, MipIndex, UpdateRegion, SourceRowPitch, SourceDepthPitch, SourceData, RHITSourceData](FRHICommandListBase&)
+	{
+		VERIFY_GL_SCOPE();
+		check(FOpenGL::SupportsTexture3D());
+		FOpenGLTexture* Texture = GetOpenGLTextureFromRHITexture(TextureRHI);
 
-	EPixelFormat PixelFormat = Texture->GetFormat();
-	const FOpenGLTextureFormat& GLFormat = GOpenGLTextureFormats[PixelFormat];
-	const FPixelFormatInfo& FormatInfo = GPixelFormats[PixelFormat];
-	const uint32 FormatBPP = FormatInfo.BlockBytes;
-
-	check( FOpenGL::SupportsTexture3D() );
-	// TO DO - add appropriate offsets to source data when necessary
-	check(UpdateRegion.SrcX == 0);
-	check(UpdateRegion.SrcY == 0);
-	check(UpdateRegion.SrcZ == 0);
-	
-	glPixelStorei(GL_UNPACK_ALIGNMENT, 1);
+		// Use a texture stage that's not likely to be used for draws, to avoid waiting
+		FOpenGLContextState& ContextState = GetContextStateForCurrentContext();
+		CachedSetupTextureStage(ContextState, FOpenGL::GetMaxCombinedTextureImageUnits() - 1, Texture->Target, Texture->GetResource(), 0, Texture->GetNumMips());
+		CachedBindPixelUnpackBuffer(ContextState, 0);
 
-	const bool bSRGB = EnumHasAnyFlags(Texture->GetFlags(), TexCreate_SRGB);
+		EPixelFormat PixelFormat = Texture->GetFormat();
+		const FOpenGLTextureFormat& GLFormat = GOpenGLTextureFormats[PixelFormat];
+		const FPixelFormatInfo& FormatInfo = GPixelFormats[PixelFormat];
+		const uint32 FormatBPP = FormatInfo.BlockBytes;
 
-	if (GLFormat.bCompressed)
-	{
-		FOpenGL::CompressedTexSubImage3D(
-			Texture->Target,
-			MipIndex,
-			UpdateRegion.DestX, UpdateRegion.DestY, UpdateRegion.DestZ,
-			UpdateRegion.Width, UpdateRegion.Height, UpdateRegion.Depth,
-			GLFormat.InternalFormat[bSRGB],
-			SourceDepthPitch * UpdateRegion.Depth,
-			SourceData);
-	}
-	else
-	{
-		glPixelStorei(GL_UNPACK_ROW_LENGTH, UpdateRegion.Width / FormatInfo.BlockSizeX);
-		glPixelStorei(GL_UNPACK_IMAGE_HEIGHT, UpdateRegion.Height / FormatInfo.BlockSizeY);
+		check(FOpenGL::SupportsTexture3D());
+		// TO DO - add appropriate offsets to source data when necessary
+		check(UpdateRegion.SrcX == 0);
+		check(UpdateRegion.SrcY == 0);
+		check(UpdateRegion.SrcZ == 0);
 
-		FOpenGL::TexSubImage3D(Texture->Target, MipIndex, UpdateRegion.DestX, UpdateRegion.DestY, UpdateRegion.DestZ, UpdateRegion.Width, UpdateRegion.Height, UpdateRegion.Depth, GLFormat.Format, GLFormat.Type, SourceData);
-	}
+		glPixelStorei(GL_UNPACK_ALIGNMENT, 1);
+
+		const bool bSRGB = EnumHasAnyFlags(Texture->GetFlags(), TexCreate_SRGB);
 
-	glPixelStorei(GL_UNPACK_ALIGNMENT, 4);
-	glPixelStorei(GL_UNPACK_IMAGE_HEIGHT, 0);
-	glPixelStorei(GL_UNPACK_ROW_LENGTH, 0);
+		if (GLFormat.bCompressed)
+		{
+			FOpenGL::CompressedTexSubImage3D(
+				Texture->Target,
+				MipIndex,
+				UpdateRegion.DestX, UpdateRegion.DestY, UpdateRegion.DestZ,
+				UpdateRegion.Width, UpdateRegion.Height, UpdateRegion.Depth,
+				GLFormat.InternalFormat[bSRGB],
+				SourceDepthPitch * UpdateRegion.Depth,
+				SourceData);
+		}
+		else
+		{
+			glPixelStorei(GL_UNPACK_ROW_LENGTH, UpdateRegion.Width / FormatInfo.BlockSizeX);
+			glPixelStorei(GL_UNPACK_IMAGE_HEIGHT, UpdateRegion.Height / FormatInfo.BlockSizeY);
 
+			FOpenGL::TexSubImage3D(Texture->Target, MipIndex, UpdateRegion.DestX, UpdateRegion.DestY, UpdateRegion.DestZ, UpdateRegion.Width, UpdateRegion.Height, UpdateRegion.Depth, GLFormat.Format, GLFormat.Type, SourceData);
+		}
 
-	// No need to restore texture stage; leave it like this,
-	// and the next draw will take care of cleaning it up; or
-	// next operation that needs the stage will switch something else in on it.
+		glPixelStorei(GL_UNPACK_ALIGNMENT, 4);
+		glPixelStorei(GL_UNPACK_IMAGE_HEIGHT, 0);
+		glPixelStorei(GL_UNPACK_ROW_LENGTH, 0);
+
+		// free source data if we're on RHIT
+		if (RHITSourceData)
+		{
+			FMemory::Free(RHITSourceData);
+		}
+
+		// No need to restore texture stage; leave it like this,
+		// and the next draw will take care of cleaning it up; or
+		// next operation that needs the stage will switch something else in on it.
+	});
 }
 
 void FOpenGLDynamicRHI::InvalidateTextureResourceInCache(GLuint Resource)
diff --git a/Engine/Source/Runtime/OpenGLDrv/Public/OpenGLDrv.h b/Engine/Source/Runtime/OpenGLDrv/Public/OpenGLDrv.h
--- a/Engine/Source/Runtime/OpenGLDrv/Public/OpenGLDrv.h
+++ b/Engine/Source/Runtime/OpenGLDrv/Public/OpenGLDrv.h
@@ -400,8 +400,8 @@ public:
 	virtual void RHIUnlockTexture2D(FRHITexture2D* Texture, uint32 MipIndex, bool bLockWithinMiptail) final override;
 	virtual void* RHILockTexture2DArray(FRHITexture2DArray* Texture, uint32 TextureIndex, uint32 MipIndex, EResourceLockMode LockMode, uint32& DestStride, bool bLockWithinMiptail) final override;
 	virtual void RHIUnlockTexture2DArray(FRHITexture2DArray* Texture, uint32 TextureIndex, uint32 MipIndex, bool bLockWithinMiptail) final override;
-	virtual void RHIUpdateTexture2D(FRHITexture2D* Texture, uint32 MipIndex, const struct FUpdateTextureRegion2D& UpdateRegion, uint32 SourcePitch, const uint8* SourceData) final override;
-	virtual void RHIUpdateTexture3D(FRHITexture3D* Texture, uint32 MipIndex, const struct FUpdateTextureRegion3D& UpdateRegion, uint32 SourceRowPitch, uint32 SourceDepthPitch, const uint8* SourceData) final override;
+	virtual void RHIUpdateTexture2D(FRHICommandListBase& RHICmdList, FRHITexture2D* Texture, uint32 MipIndex, const struct FUpdateTextureRegion2D& UpdateRegion, uint32 SourcePitch, const uint8* SourceData) final override;
+	virtual void RHIUpdateTexture3D(FRHICommandListBase& RHICmdList, FRHITexture3D* Texture, uint32 MipIndex, const struct FUpdateTextureRegion3D& UpdateRegion, uint32 SourceRowPitch, uint32 SourceDepthPitch, const uint8* SourceData) final override;
 	virtual void* RHILockTextureCubeFace(FRHITextureCube* Texture, uint32 FaceIndex, uint32 ArrayIndex, uint32 MipIndex, EResourceLockMode LockMode, uint32& DestStride, bool bLockWithinMiptail) final override;
 	virtual void RHIUnlockTextureCubeFace(FRHITextureCube* Texture, uint32 FaceIndex, uint32 ArrayIndex, uint32 MipIndex, bool bLockWithinMiptail) final override;
 	virtual void RHIBindDebugLabelName(FRHITexture* Texture, const TCHAR* Name) final override;
@@ -861,19 +861,6 @@ public:
 
 	virtual FSamplerStateRHIRef RHICreateSamplerState(const FSamplerStateInitializerRHI& Initializer) final override;
 
-	virtual void UpdateTexture2D_RenderThread(class FRHICommandListImmediate& RHICmdList, FRHITexture2D* Texture, uint32 MipIndex, const struct FUpdateTextureRegion2D& UpdateRegion, uint32 SourcePitch, const uint8* SourceData) final override
-	{
-		GDynamicRHI->RHIUpdateTexture2D(Texture, MipIndex, UpdateRegion, SourcePitch, SourceData);
-	}
-
-	virtual void UpdateTexture3D_RenderThread(class FRHICommandListImmediate& RHICmdList, FRHITexture3D* Texture, uint32 MipIndex, const struct FUpdateTextureRegion3D& UpdateRegion, uint32 SourceRowPitch, uint32 SourceDepthPitch, const uint8* SourceData) final override
-	{
-		RHITHREAD_GLCOMMAND_PROLOGUE();
-		GDynamicRHI->RHIUpdateTexture3D(Texture, MipIndex, UpdateRegion, SourceRowPitch, SourceDepthPitch, SourceData);
-		RHITHREAD_GLCOMMAND_EPILOGUE();
-	}
-
-
 	virtual FGraphicsPipelineStateRHIRef RHICreateGraphicsPipelineState(const FGraphicsPipelineStateInitializer& Initializer)
 	{
 		PrepareGFXBoundShaderState(Initializer);
diff --git a/Engine/Source/Runtime/RHI/Private/RHICommandList.cpp b/Engine/Source/Runtime/RHI/Private/RHICommandList.cpp
--- a/Engine/Source/Runtime/RHI/Private/RHICommandList.cpp
+++ b/Engine/Source/Runtime/RHI/Private/RHICommandList.cpp
@@ -1639,22 +1639,9 @@ FComputeShaderRHIRef FDynamicRHI::CreateComputeShader_RenderThread(class FRHICom
 	return GDynamicRHI->RHICreateComputeShader(Code, Hash);
 }
 
-void FDynamicRHI::UpdateTexture2D_RenderThread(class FRHICommandListImmediate& RHICmdList, FRHITexture2D* Texture, uint32 MipIndex, const struct FUpdateTextureRegion2D& UpdateRegion, uint32 SourcePitch, const uint8* SourceData)
+FUpdateTexture3DData FDynamicRHI::RHIBeginUpdateTexture3D(FRHICommandListBase& RHICmdList, FRHITexture3D* Texture, uint32 MipIndex, const struct FUpdateTextureRegion3D& UpdateRegion)
 {
-	CSV_SCOPED_TIMING_STAT(RHITStalls, UpdateTexture2D_RenderThread);
-	FScopedRHIThreadStaller StallRHIThread(RHICmdList);
-	return GDynamicRHI->RHIUpdateTexture2D(Texture, MipIndex, UpdateRegion, SourcePitch, SourceData);
-}
-
-void FDynamicRHI::UpdateFromBufferTexture2D_RenderThread(class FRHICommandListImmediate& RHICmdList, FRHITexture2D* Texture, uint32 MipIndex, const struct FUpdateTextureRegion2D& UpdateRegion, uint32 SourcePitch, FRHIBuffer* Buffer, uint32 BufferOffset)
-{
-	FScopedRHIThreadStaller StallRHIThread(RHICmdList);
-	return GDynamicRHI->RHIUpdateFromBufferTexture2D(Texture, MipIndex, UpdateRegion, SourcePitch, Buffer, BufferOffset);
-}
-
-FUpdateTexture3DData FDynamicRHI::BeginUpdateTexture3D_RenderThread(class FRHICommandListImmediate& RHICmdList, FRHITexture3D* Texture, uint32 MipIndex, const struct FUpdateTextureRegion3D& UpdateRegion)
-{
-	check(IsInRenderingThread());
+	check(IsInParallelRenderingThread());
 
 	const int32 FormatSize = PixelFormatBlockBytes[Texture->GetFormat()];
 	const int32 RowPitch = UpdateRegion.Width * FormatSize;
@@ -1666,32 +1653,23 @@ FUpdateTexture3DData FDynamicRHI::BeginUpdateTexture3D_RenderThread(class FRHICo
 	return FUpdateTexture3DData(Texture, MipIndex, UpdateRegion, RowPitch, DepthPitch, Data, MemorySize, GFrameNumberRenderThread);
 }
 
-void FDynamicRHI::EndUpdateTexture3D_RenderThread(class FRHICommandListImmediate& RHICmdList, FUpdateTexture3DData& UpdateData)
+void FDynamicRHI::RHIEndUpdateTexture3D(FRHICommandListBase& RHICmdList, FUpdateTexture3DData& UpdateData)
 {
-	check(IsInRenderingThread());
+	check(IsInParallelRenderingThread());
 	check(GFrameNumberRenderThread == UpdateData.FrameNumber); 
-	CSV_SCOPED_TIMING_STAT(RHITStalls, EndUpdateTexture3D_RenderThread);
-	FScopedRHIThreadStaller StallRHIThread(RHICmdList);	
-	GDynamicRHI->RHIUpdateTexture3D(UpdateData.Texture, UpdateData.MipIndex, UpdateData.UpdateRegion, UpdateData.RowPitch, UpdateData.DepthPitch, UpdateData.Data);
+	GDynamicRHI->RHIUpdateTexture3D(RHICmdList, UpdateData.Texture, UpdateData.MipIndex, UpdateData.UpdateRegion, UpdateData.RowPitch, UpdateData.DepthPitch, UpdateData.Data);
 	FMemory::Free(UpdateData.Data);
 	UpdateData.Data = nullptr;
 }
 
-void FDynamicRHI::EndMultiUpdateTexture3D_RenderThread(class FRHICommandListImmediate& RHICmdList, TArray<FUpdateTexture3DData>& UpdateDataArray)
+void FDynamicRHI::RHIEndMultiUpdateTexture3D(FRHICommandListBase& RHICmdList, TArray<FUpdateTexture3DData>& UpdateDataArray)
 {
 	for (int32 Idx = 0; Idx < UpdateDataArray.Num(); ++Idx)
 	{
-		GDynamicRHI->EndUpdateTexture3D_RenderThread(RHICmdList, UpdateDataArray[Idx]);
+		GDynamicRHI->RHIEndUpdateTexture3D(RHICmdList, UpdateDataArray[Idx]);
 	}
 }
 
-void FDynamicRHI::UpdateTexture3D_RenderThread(class FRHICommandListImmediate& RHICmdList, FRHITexture3D* Texture, uint32 MipIndex, const struct FUpdateTextureRegion3D& UpdateRegion, uint32 SourceRowPitch, uint32 SourceDepthPitch, const uint8* SourceData)
-{
-	CSV_SCOPED_TIMING_STAT(RHITStalls, UpdateTexture3D_RenderThread);
-	FScopedRHIThreadStaller StallRHIThread(RHICmdList);
-	GDynamicRHI->RHIUpdateTexture3D(Texture, MipIndex, UpdateRegion, SourceRowPitch, SourceDepthPitch, SourceData);
-}
-
 void* FDynamicRHI::LockTexture2D_RenderThread(class FRHICommandListImmediate& RHICmdList, FRHITexture2D* Texture, uint32 MipIndex, EResourceLockMode LockMode, uint32& DestStride, bool bLockWithinMiptail, bool bNeedsDefaultRHIFlush)
 {
 	if (bNeedsDefaultRHIFlush) 
diff --git a/Engine/Source/Runtime/RHI/Public/DynamicRHI.h b/Engine/Source/Runtime/RHI/Public/DynamicRHI.h
--- a/Engine/Source/Runtime/RHI/Public/DynamicRHI.h
+++ b/Engine/Source/Runtime/RHI/Public/DynamicRHI.h
@@ -684,8 +684,7 @@ public:
 	* @param SourcePitch - size in bytes of each row of the source image
 	* @param SourceData - source image data, starting at the upper left corner of the source rectangle (in same pixel format as texture)
 	*/
-	// FlushType: Flush RHI Thread
-	virtual void RHIUpdateTexture2D(FRHITexture2D* Texture, uint32 MipIndex, const struct FUpdateTextureRegion2D& UpdateRegion, uint32 SourcePitch, const uint8* SourceData) = 0;
+	virtual void RHIUpdateTexture2D(FRHICommandListBase& RHICmdList, FRHITexture2D* Texture, uint32 MipIndex, const struct FUpdateTextureRegion2D& UpdateRegion, uint32 SourcePitch, const uint8* SourceData) = 0;
 
 	/**
 	* Updates a region of a 2D texture from GPU memory provided by the given buffer (may not be implemented on every platform)
@@ -695,8 +694,10 @@ public:
 	* @param SourcePitch - size in bytes of each row of the source image
 	* @param Buffer, BufferOffset - source image data, starting at the upper left corner of the source rectangle (in same pixel format as texture)
 	*/
-	// FlushType: Flush RHI Thread
-	virtual void RHIUpdateFromBufferTexture2D(FRHITexture2D* Texture, uint32 MipIndex, const struct FUpdateTextureRegion2D& UpdateRegion, uint32 SourcePitch, FRHIBuffer* Buffer, uint32 BufferOffset) {}
+	virtual void RHIUpdateFromBufferTexture2D(FRHICommandListBase& RHICmdList, FRHITexture2D* Texture, uint32 MipIndex, const struct FUpdateTextureRegion2D& UpdateRegion, uint32 SourcePitch, FRHIBuffer* Buffer, uint32 BufferOffset)
+	{
+		checkNoEntry();
+	}
 
 	/**
 	* Updates a region of a 3D texture from system memory
@@ -707,8 +708,7 @@ public:
 	* @param SourceDepthPitch - size in bytes of each depth slice of the source image, usually Bpp * SizeX * SizeY
 	* @param SourceData - source image data, starting at the upper left corner of the source rectangle (in same pixel format as texture)
 	*/
-	// FlushType: Flush RHI Thread
-	virtual void RHIUpdateTexture3D(FRHITexture3D* Texture, uint32 MipIndex, const struct FUpdateTextureRegion3D& UpdateRegion, uint32 SourceRowPitch, uint32 SourceDepthPitch, const uint8* SourceData) = 0;
+	virtual void RHIUpdateTexture3D(FRHICommandListBase& RHICmdList, FRHITexture3D* Texture, uint32 MipIndex, const struct FUpdateTextureRegion3D& UpdateRegion, uint32 SourceRowPitch, uint32 SourceDepthPitch, const uint8* SourceData) = 0;
 
 	/**
 	* Locks an RHI texture's mip-map for read/write operations on the CPU
@@ -1077,17 +1077,14 @@ public:
 	virtual FComputeShaderRHIRef CreateComputeShader_RenderThread(class FRHICommandListImmediate& RHICmdList, TArrayView<const uint8> Code, const FSHAHash& Hash);
 	virtual void* LockTexture2D_RenderThread(class FRHICommandListImmediate& RHICmdList, FRHITexture2D* Texture, uint32 MipIndex, EResourceLockMode LockMode, uint32& DestStride, bool bLockWithinMiptail, bool bNeedsDefaultRHIFlush = true);
 	virtual void UnlockTexture2D_RenderThread(class FRHICommandListImmediate& RHICmdList, FRHITexture2D* Texture, uint32 MipIndex, bool bLockWithinMiptail, bool bNeedsDefaultRHIFlush = true);
-	virtual void UpdateTexture2D_RenderThread(class FRHICommandListImmediate& RHICmdList, FRHITexture2D* Texture, uint32 MipIndex, const struct FUpdateTextureRegion2D& UpdateRegion, uint32 SourcePitch, const uint8* SourceData);
-	virtual void UpdateFromBufferTexture2D_RenderThread(class FRHICommandListImmediate& RHICmdList, FRHITexture2D* Texture, uint32 MipIndex, const struct FUpdateTextureRegion2D& UpdateRegion, uint32 SourcePitch, FRHIBuffer* Buffer, uint32 BufferOffset);
+
 	virtual void* LockTexture2DArray_RenderThread(class FRHICommandListImmediate& RHICmdList, FRHITexture2DArray* Texture, uint32 ArrayIndex, uint32 MipIndex, EResourceLockMode LockMode, uint32& DestStride, bool bLockWithinMiptail);
 	virtual void UnlockTexture2DArray_RenderThread(class FRHICommandListImmediate& RHICmdList, FRHITexture2DArray* Texture, uint32 ArrayIndex, uint32 MipIndex, bool bLockWithinMiptail);
 
-	virtual FUpdateTexture3DData BeginUpdateTexture3D_RenderThread(class FRHICommandListImmediate& RHICmdList, FRHITexture3D* Texture, uint32 MipIndex, const struct FUpdateTextureRegion3D& UpdateRegion);
-	virtual void EndUpdateTexture3D_RenderThread(class FRHICommandListImmediate& RHICmdList, FUpdateTexture3DData& UpdateData);
-
-	virtual void EndMultiUpdateTexture3D_RenderThread(class FRHICommandListImmediate& RHICmdList, TArray<FUpdateTexture3DData>& UpdateDataArray);
+	virtual FUpdateTexture3DData RHIBeginUpdateTexture3D(FRHICommandListBase& RHICmdList, FRHITexture3D* Texture, uint32 MipIndex, const struct FUpdateTextureRegion3D& UpdateRegion);
+	virtual void RHIEndUpdateTexture3D(FRHICommandListBase& RHICmdList, FUpdateTexture3DData& UpdateData);
 
-	virtual void UpdateTexture3D_RenderThread(class FRHICommandListImmediate& RHICmdList, FRHITexture3D* Texture, uint32 MipIndex, const struct FUpdateTextureRegion3D& UpdateRegion, uint32 SourceRowPitch, uint32 SourceDepthPitch, const uint8* SourceData);
+	virtual void RHIEndMultiUpdateTexture3D(FRHICommandListBase& RHICmdList, TArray<FUpdateTexture3DData>& UpdateDataArray);
 
 	virtual FRHIShaderLibraryRef RHICreateShaderLibrary_RenderThread(class FRHICommandListImmediate& RHICmdList, EShaderPlatform Platform, FString FilePath, FString Name);
 
diff --git a/Engine/Source/Runtime/RHI/Public/RHICommandList.h b/Engine/Source/Runtime/RHI/Public/RHICommandList.h
--- a/Engine/Source/Runtime/RHI/Public/RHICommandList.h
+++ b/Engine/Source/Runtime/RHI/Public/RHICommandList.h
@@ -406,6 +406,16 @@ struct TRHILambdaCommand final : public FRHICommandBase
 
 #define RHI_COUNT_COMMANDS (DO_CHECK || STATS)
 
+class FRHICommandListScopedPipelineGuard
+{
+	FRHICommandListBase& RHICmdList;
+	bool bPipelineSet = false;
+
+public:
+	FRHICommandListScopedPipelineGuard(FRHICommandListBase& RHICmdList);
+	~FRHICommandListScopedPipelineGuard();
+};
+
 class RHI_API FRHICommandListBase : public FNoncopyable
 {
 public:
@@ -591,16 +601,19 @@ public:
 
 	FORCEINLINE void* LockBuffer(FRHIBuffer* Buffer, uint32 Offset, uint32 SizeRHI, EResourceLockMode LockMode)
 	{
+		FRHICommandListScopedPipelineGuard ScopedPipeline(*this);
 		return GDynamicRHI->RHILockBuffer(*this, Buffer, Offset, SizeRHI, LockMode);
 	}
 
 	FORCEINLINE void UnlockBuffer(FRHIBuffer* Buffer)
 	{
+		FRHICommandListScopedPipelineGuard ScopedPipeline(*this);
 		GDynamicRHI->RHIUnlockBuffer(*this, Buffer);
 	}
 
 	FORCEINLINE FBufferRHIRef CreateBuffer(uint32 Size, EBufferUsageFlags Usage, uint32 Stride, ERHIAccess ResourceState, FRHIResourceCreateInfo& CreateInfo)
 	{
+		FRHICommandListScopedPipelineGuard ScopedPipeline(*this);
 		FBufferRHIRef Buffer = GDynamicRHI->RHICreateBuffer(*this, Size, Usage, Stride, ResourceState, CreateInfo);
 		Buffer->SetTrackedAccess_Unsafe(ResourceState);
 		return Buffer;
@@ -608,9 +621,41 @@ public:
 
 	FORCEINLINE void UpdateUniformBuffer(FRHIUniformBuffer* UniformBufferRHI, const void* Contents)
 	{
+		FRHICommandListScopedPipelineGuard ScopedPipeline(*this);
 		GDynamicRHI->RHIUpdateUniformBuffer(*this, UniformBufferRHI, Contents);
 	}
 
+	FORCEINLINE void UpdateTexture2D(FRHITexture2D* Texture, uint32 MipIndex, const struct FUpdateTextureRegion2D& UpdateRegion, uint32 SourcePitch, const uint8* SourceData)
+	{
+		checkf(UpdateRegion.DestX + UpdateRegion.Width <= Texture->GetSizeX(), TEXT("UpdateTexture2D out of bounds on X. Texture: %s, %i, %i, %i"), *Texture->GetName().ToString(), UpdateRegion.DestX, UpdateRegion.Width, Texture->GetSizeX());
+		checkf(UpdateRegion.DestY + UpdateRegion.Height <= Texture->GetSizeY(), TEXT("UpdateTexture2D out of bounds on Y. Texture: %s, %i, %i, %i"), *Texture->GetName().ToString(), UpdateRegion.DestY, UpdateRegion.Height, Texture->GetSizeY());
+		LLM_SCOPE(ELLMTag::Textures);
+
+		FRHICommandListScopedPipelineGuard ScopedPipeline(*this);
+		GDynamicRHI->RHIUpdateTexture2D(*this, Texture, MipIndex, UpdateRegion, SourcePitch, SourceData);
+	}
+
+	FORCEINLINE void UpdateFromBufferTexture2D(FRHITexture2D* Texture, uint32 MipIndex, const struct FUpdateTextureRegion2D& UpdateRegion, uint32 SourcePitch, FRHIBuffer* Buffer, uint32 BufferOffset)
+	{
+		checkf(UpdateRegion.DestX + UpdateRegion.Width <= Texture->GetSizeX(), TEXT("UpdateFromBufferTexture2D out of bounds on X. Texture: %s, %i, %i, %i"), *Texture->GetName().ToString(), UpdateRegion.DestX, UpdateRegion.Width, Texture->GetSizeX());
+		checkf(UpdateRegion.DestY + UpdateRegion.Height <= Texture->GetSizeY(), TEXT("UpdateFromBufferTexture2D out of bounds on Y. Texture: %s, %i, %i, %i"), *Texture->GetName().ToString(), UpdateRegion.DestY, UpdateRegion.Height, Texture->GetSizeY());
+		LLM_SCOPE(ELLMTag::Textures);
+
+		FRHICommandListScopedPipelineGuard ScopedPipeline(*this);
+		GDynamicRHI->RHIUpdateFromBufferTexture2D(*this, Texture, MipIndex, UpdateRegion, SourcePitch, Buffer, BufferOffset);
+	}
+
+	FORCEINLINE void UpdateTexture3D(FRHITexture3D* Texture, uint32 MipIndex, const struct FUpdateTextureRegion3D& UpdateRegion, uint32 SourceRowPitch, uint32 SourceDepthPitch, const uint8* SourceData)
+	{
+		checkf(UpdateRegion.DestX + UpdateRegion.Width <= Texture->GetSizeX(), TEXT("UpdateTexture3D out of bounds on X. Texture: %s, %i, %i, %i"), *Texture->GetName().ToString(), UpdateRegion.DestX, UpdateRegion.Width, Texture->GetSizeX());
+		checkf(UpdateRegion.DestY + UpdateRegion.Height <= Texture->GetSizeY(), TEXT("UpdateTexture3D out of bounds on Y. Texture: %s, %i, %i, %i"), *Texture->GetName().ToString(), UpdateRegion.DestY, UpdateRegion.Height, Texture->GetSizeY());
+		checkf(UpdateRegion.DestZ + UpdateRegion.Depth <= Texture->GetSizeZ(), TEXT("UpdateTexture3D out of bounds on Z. Texture: %s, %i, %i, %i"), *Texture->GetName().ToString(), UpdateRegion.DestZ, UpdateRegion.Depth, Texture->GetSizeZ());
+		LLM_SCOPE(ELLMTag::Textures);
+
+		FRHICommandListScopedPipelineGuard ScopedPipeline(*this);
+		GDynamicRHI->RHIUpdateTexture3D(*this, Texture, MipIndex, UpdateRegion, SourceRowPitch, SourceDepthPitch, SourceData);
+	}
+
 protected:
 	FMemStackBase& GetAllocator() { return MemManager; }
 
@@ -4264,52 +4309,27 @@ public:
 		GDynamicRHI->UnlockTexture2DArray_RenderThread(*this, Texture, TextureIndex, MipIndex, bLockWithinMiptail);
 	}
 	
-	FORCEINLINE void UpdateTexture2D(FRHITexture2D* Texture, uint32 MipIndex, const struct FUpdateTextureRegion2D& UpdateRegion, uint32 SourcePitch, const uint8* SourceData)
-	{		
-		checkf(UpdateRegion.DestX + UpdateRegion.Width <= Texture->GetSizeX(), TEXT("UpdateTexture2D out of bounds on X. Texture: %s, %i, %i, %i"), *Texture->GetName().ToString(), UpdateRegion.DestX, UpdateRegion.Width, Texture->GetSizeX());
-		checkf(UpdateRegion.DestY + UpdateRegion.Height <= Texture->GetSizeY(), TEXT("UpdateTexture2D out of bounds on Y. Texture: %s, %i, %i, %i"), *Texture->GetName().ToString(), UpdateRegion.DestY, UpdateRegion.Height, Texture->GetSizeY());
-		LLM_SCOPE(ELLMTag::Textures);
-		GDynamicRHI->UpdateTexture2D_RenderThread(*this, Texture, MipIndex, UpdateRegion, SourcePitch, SourceData);
-	}
-
-	FORCEINLINE void UpdateFromBufferTexture2D(FRHITexture2D* Texture, uint32 MipIndex, const struct FUpdateTextureRegion2D& UpdateRegion, uint32 SourcePitch, FRHIBuffer* Buffer, uint32 BufferOffset)
-	{
-		checkf(UpdateRegion.DestX + UpdateRegion.Width <= Texture->GetSizeX(), TEXT("UpdateFromBufferTexture2D out of bounds on X. Texture: %s, %i, %i, %i"), *Texture->GetName().ToString(), UpdateRegion.DestX, UpdateRegion.Width, Texture->GetSizeX());
-		checkf(UpdateRegion.DestY + UpdateRegion.Height <= Texture->GetSizeY(), TEXT("UpdateFromBufferTexture2D out of bounds on Y. Texture: %s, %i, %i, %i"), *Texture->GetName().ToString(), UpdateRegion.DestY, UpdateRegion.Height, Texture->GetSizeY());
-		LLM_SCOPE(ELLMTag::Textures);
-		GDynamicRHI->UpdateFromBufferTexture2D_RenderThread(*this, Texture, MipIndex, UpdateRegion, SourcePitch, Buffer, BufferOffset);
-	}
-
 	FORCEINLINE FUpdateTexture3DData BeginUpdateTexture3D(FRHITexture3D* Texture, uint32 MipIndex, const struct FUpdateTextureRegion3D& UpdateRegion)
 	{
 		checkf(UpdateRegion.DestX + UpdateRegion.Width <= Texture->GetSizeX(), TEXT("UpdateTexture3D out of bounds on X. Texture: %s, %i, %i, %i"), *Texture->GetName().ToString(), UpdateRegion.DestX, UpdateRegion.Width, Texture->GetSizeX());
 		checkf(UpdateRegion.DestY + UpdateRegion.Height <= Texture->GetSizeY(), TEXT("UpdateTexture3D out of bounds on Y. Texture: %s, %i, %i, %i"), *Texture->GetName().ToString(), UpdateRegion.DestY, UpdateRegion.Height, Texture->GetSizeY());
 		checkf(UpdateRegion.DestZ + UpdateRegion.Depth <= Texture->GetSizeZ(), TEXT("UpdateTexture3D out of bounds on Z. Texture: %s, %i, %i, %i"), *Texture->GetName().ToString(), UpdateRegion.DestZ, UpdateRegion.Depth, Texture->GetSizeZ());
 		LLM_SCOPE(ELLMTag::Textures);
-		return GDynamicRHI->BeginUpdateTexture3D_RenderThread(*this, Texture, MipIndex, UpdateRegion);
+		return GDynamicRHI->RHIBeginUpdateTexture3D(*this, Texture, MipIndex, UpdateRegion);
 	}
 
 	FORCEINLINE void EndUpdateTexture3D(FUpdateTexture3DData& UpdateData)
-	{		
+	{
 		LLM_SCOPE(ELLMTag::Textures);
-		GDynamicRHI->EndUpdateTexture3D_RenderThread(*this, UpdateData);
+		GDynamicRHI->RHIEndUpdateTexture3D(*this, UpdateData);
 	}
 
 	FORCEINLINE void EndMultiUpdateTexture3D(TArray<FUpdateTexture3DData>& UpdateDataArray)
 	{
 		LLM_SCOPE(ELLMTag::Textures);
-		GDynamicRHI->EndMultiUpdateTexture3D_RenderThread(*this, UpdateDataArray);
+		GDynamicRHI->RHIEndMultiUpdateTexture3D(*this, UpdateDataArray);
 	}
-	
-	FORCEINLINE void UpdateTexture3D(FRHITexture3D* Texture, uint32 MipIndex, const struct FUpdateTextureRegion3D& UpdateRegion, uint32 SourceRowPitch, uint32 SourceDepthPitch, const uint8* SourceData)
-	{
-		checkf(UpdateRegion.DestX + UpdateRegion.Width <= Texture->GetSizeX(), TEXT("UpdateTexture3D out of bounds on X. Texture: %s, %i, %i, %i"), *Texture->GetName().ToString(), UpdateRegion.DestX, UpdateRegion.Width, Texture->GetSizeX());
-		checkf(UpdateRegion.DestY + UpdateRegion.Height <= Texture->GetSizeY(), TEXT("UpdateTexture3D out of bounds on Y. Texture: %s, %i, %i, %i"), *Texture->GetName().ToString(), UpdateRegion.DestY, UpdateRegion.Height, Texture->GetSizeY());
-		checkf(UpdateRegion.DestZ + UpdateRegion.Depth <= Texture->GetSizeZ(), TEXT("UpdateTexture3D out of bounds on Z. Texture: %s, %i, %i, %i"), *Texture->GetName().ToString(), UpdateRegion.DestZ, UpdateRegion.Depth, Texture->GetSizeZ());
-		LLM_SCOPE(ELLMTag::Textures);
-		GDynamicRHI->UpdateTexture3D_RenderThread(*this, Texture, MipIndex, UpdateRegion, SourceRowPitch, SourceDepthPitch, SourceData);
-	}
-	
+
 	FORCEINLINE void* LockTextureCubeFace(FRHITexture* Texture, uint32 FaceIndex, uint32 ArrayIndex, uint32 MipIndex, EResourceLockMode LockMode, uint32& DestStride, bool bLockWithinMiptail)
 	{
 		LLM_SCOPE(ELLMTag::Textures);
@@ -4633,7 +4653,8 @@ public:
 	FRHICommandListScopedPipeline(FRHICommandListBase& RHICmdList, ERHIPipeline Pipeline)
 		: RHICmdList(RHICmdList)
 		, PreviousPipeline(RHICmdList.SwitchPipeline(Pipeline))
-	{}
+	{
+	}
 
 	~FRHICommandListScopedPipeline()
 	{
diff --git a/Engine/Source/Runtime/RHI/Public/RHICommandList.inl b/Engine/Source/Runtime/RHI/Public/RHICommandList.inl
--- a/Engine/Source/Runtime/RHI/Public/RHICommandList.inl
+++ b/Engine/Source/Runtime/RHI/Public/RHICommandList.inl
@@ -145,3 +145,20 @@ private:
 #endif
 };
 
+inline FRHICommandListScopedPipelineGuard::FRHICommandListScopedPipelineGuard(FRHICommandListBase& RHICmdList)
+	: RHICmdList(RHICmdList)
+{
+	if (RHICmdList.GetPipeline() == ERHIPipeline::None)
+	{
+		RHICmdList.SwitchPipeline(ERHIPipeline::Graphics);
+		bPipelineSet = true;
+	}
+}
+
+inline FRHICommandListScopedPipelineGuard::~FRHICommandListScopedPipelineGuard()
+{
+	if (bPipelineSet)
+	{
+		RHICmdList.SwitchPipeline(ERHIPipeline::None);
+	}
+}
\ No newline at end of file
diff --git a/Engine/Source/Runtime/RHI/Public/RHIValidation.h b/Engine/Source/Runtime/RHI/Public/RHIValidation.h
--- a/Engine/Source/Runtime/RHI/Public/RHIValidation.h
+++ b/Engine/Source/Runtime/RHI/Public/RHIValidation.h
@@ -798,14 +798,14 @@ public:
 	* @param SourceData - source image data, starting at the upper left corner of the source rectangle (in same pixel format as texture)
 	*/
 	// FlushType: Flush RHI Thread
-	virtual void RHIUpdateTexture2D(FRHITexture2D* Texture, uint32 MipIndex, const struct FUpdateTextureRegion2D& UpdateRegion, uint32 SourcePitch, const uint8* SourceData) override final
+	virtual void RHIUpdateTexture2D(FRHICommandListBase& RHICmdList, FRHITexture2D* Texture, uint32 MipIndex, const struct FUpdateTextureRegion2D& UpdateRegion, uint32 SourcePitch, const uint8* SourceData) override final
 	{
-		RHI->RHIUpdateTexture2D(Texture, MipIndex, UpdateRegion, SourcePitch, SourceData);
+		RHI->RHIUpdateTexture2D(RHICmdList, Texture, MipIndex, UpdateRegion, SourcePitch, SourceData);
 	}
 
-	virtual void RHIUpdateFromBufferTexture2D(FRHITexture2D* Texture, uint32 MipIndex, const struct FUpdateTextureRegion2D& UpdateRegion, uint32 SourcePitch, FRHIBuffer* Buffer, uint32 BufferOffset) override final
+	virtual void RHIUpdateFromBufferTexture2D(FRHICommandListBase& RHICmdList, FRHITexture2D* Texture, uint32 MipIndex, const struct FUpdateTextureRegion2D& UpdateRegion, uint32 SourcePitch, FRHIBuffer* Buffer, uint32 BufferOffset) override final
 	{
-		RHI->RHIUpdateFromBufferTexture2D(Texture, MipIndex, UpdateRegion, SourcePitch, Buffer, BufferOffset);
+		RHI->RHIUpdateFromBufferTexture2D(RHICmdList, Texture, MipIndex, UpdateRegion, SourcePitch, Buffer, BufferOffset);
 	}
 
 	/**
@@ -818,9 +818,9 @@ public:
 	* @param SourceData - source image data, starting at the upper left corner of the source rectangle (in same pixel format as texture)
 	*/
 	// FlushType: Flush RHI Thread
-	virtual void RHIUpdateTexture3D(FRHITexture3D* Texture, uint32 MipIndex, const struct FUpdateTextureRegion3D& UpdateRegion, uint32 SourceRowPitch, uint32 SourceDepthPitch, const uint8* SourceData) override final
+	virtual void RHIUpdateTexture3D(FRHICommandListBase& RHICmdList, FRHITexture3D* Texture, uint32 MipIndex, const struct FUpdateTextureRegion3D& UpdateRegion, uint32 SourceRowPitch, uint32 SourceDepthPitch, const uint8* SourceData) override final
 	{
-		RHI->RHIUpdateTexture3D(Texture, MipIndex, UpdateRegion, SourceRowPitch, SourceDepthPitch, SourceData);
+		RHI->RHIUpdateTexture3D(RHICmdList, Texture, MipIndex, UpdateRegion, SourceRowPitch, SourceDepthPitch, SourceData);
 	}
 
 	/**
@@ -1399,16 +1399,6 @@ public:
 		RHI->UnlockTexture2D_RenderThread(RHICmdList, Texture, MipIndex, bLockWithinMiptail, bNeedsDefaultRHIFlush);
 	}
 
-	virtual void UpdateTexture2D_RenderThread(class FRHICommandListImmediate& RHICmdList, FRHITexture2D* Texture, uint32 MipIndex, const struct FUpdateTextureRegion2D& UpdateRegion, uint32 SourcePitch, const uint8* SourceData) override final
-	{
-		RHI->UpdateTexture2D_RenderThread(RHICmdList, Texture, MipIndex, UpdateRegion, SourcePitch, SourceData);
-	}
-
-	virtual void UpdateFromBufferTexture2D_RenderThread(class FRHICommandListImmediate& RHICmdList, FRHITexture2D* Texture, uint32 MipIndex, const struct FUpdateTextureRegion2D& UpdateRegion, uint32 SourcePitch, FRHIBuffer* Buffer, uint32 BufferOffset) override final
-	{
-		RHI->UpdateFromBufferTexture2D_RenderThread(RHICmdList, Texture, MipIndex, UpdateRegion, SourcePitch, Buffer, BufferOffset);
-	}
-
 	virtual void* LockTexture2DArray_RenderThread(class FRHICommandListImmediate& RHICmdList, FRHITexture2DArray* Texture, uint32 ArrayIndex, uint32 MipIndex, EResourceLockMode LockMode, uint32& DestStride, bool bLockWithinMiptail) override final
 	{
 		return RHI->LockTexture2DArray_RenderThread(RHICmdList, Texture, ArrayIndex, MipIndex, LockMode, DestStride, bLockWithinMiptail);
@@ -1419,24 +1409,19 @@ public:
 		RHI->UnlockTexture2DArray_RenderThread(RHICmdList, Texture, ArrayIndex, MipIndex, bLockWithinMiptail);
 	}
 
-	virtual FUpdateTexture3DData BeginUpdateTexture3D_RenderThread(class FRHICommandListImmediate& RHICmdList, FRHITexture3D* Texture, uint32 MipIndex, const struct FUpdateTextureRegion3D& UpdateRegion) override final
-	{
-		return RHI->BeginUpdateTexture3D_RenderThread(RHICmdList, Texture, MipIndex, UpdateRegion);
-	}
-
-	virtual void EndUpdateTexture3D_RenderThread(class FRHICommandListImmediate& RHICmdList, FUpdateTexture3DData& UpdateData) override final
+	virtual FUpdateTexture3DData RHIBeginUpdateTexture3D(FRHICommandListBase& RHICmdList, FRHITexture3D* Texture, uint32 MipIndex, const struct FUpdateTextureRegion3D& UpdateRegion) override final
 	{
-		RHI->EndUpdateTexture3D_RenderThread(RHICmdList, UpdateData);
+		return RHI->RHIBeginUpdateTexture3D(RHICmdList, Texture, MipIndex, UpdateRegion);
 	}
 
-	virtual void EndMultiUpdateTexture3D_RenderThread(class FRHICommandListImmediate& RHICmdList, TArray<FUpdateTexture3DData>& UpdateDataArray) override final
+	virtual void RHIEndUpdateTexture3D(FRHICommandListBase& RHICmdList, FUpdateTexture3DData& UpdateData) override final
 	{
-		RHI->EndMultiUpdateTexture3D_RenderThread(RHICmdList, UpdateDataArray);
+		RHI->RHIEndUpdateTexture3D(RHICmdList, UpdateData);
 	}
 
-	virtual void UpdateTexture3D_RenderThread(class FRHICommandListImmediate& RHICmdList, FRHITexture3D* Texture, uint32 MipIndex, const struct FUpdateTextureRegion3D& UpdateRegion, uint32 SourceRowPitch, uint32 SourceDepthPitch, const uint8* SourceData) override final
+	virtual void RHIEndMultiUpdateTexture3D(FRHICommandListBase& RHICmdList, TArray<FUpdateTexture3DData>& UpdateDataArray) override final
 	{
-		RHI->UpdateTexture3D_RenderThread(RHICmdList, Texture, MipIndex, UpdateRegion, SourceRowPitch, SourceDepthPitch, SourceData);
+		RHI->RHIEndMultiUpdateTexture3D(RHICmdList, UpdateDataArray);
 	}
 
 	virtual FRHIShaderLibraryRef RHICreateShaderLibrary_RenderThread(class FRHICommandListImmediate& RHICmdList, EShaderPlatform Platform, FString FilePath, FString Name) override final
diff --git a/Engine/Source/Runtime/RenderCore/Private/RenderGraphBuilder.cpp b/Engine/Source/Runtime/RenderCore/Private/RenderGraphBuilder.cpp
--- a/Engine/Source/Runtime/RenderCore/Private/RenderGraphBuilder.cpp
+++ b/Engine/Source/Runtime/RenderCore/Private/RenderGraphBuilder.cpp
@@ -2470,6 +2470,7 @@ UE::Tasks::FTask FRDGBuilder::SubmitBufferUploads()
 		{
 			FTaskTagScope Scope(ETaskTag::EParallelRenderingThread);
 			RHICmdListBufferUploads = new FRHICommandList(FRHIGPUMask::All());
+			RHICmdListBufferUploads->SwitchPipeline(ERHIPipeline::Graphics);
 			SubmitUploadsLambda(*RHICmdListBufferUploads);
 			RHICmdListBufferUploads->FinishRecording();
 
diff --git a/Engine/Source/Runtime/RenderCore/Public/RenderGraphBuilder.inl b/Engine/Source/Runtime/RenderCore/Public/RenderGraphBuilder.inl
--- a/Engine/Source/Runtime/RenderCore/Public/RenderGraphBuilder.inl
+++ b/Engine/Source/Runtime/RenderCore/Public/RenderGraphBuilder.inl
@@ -427,10 +427,9 @@ void FRDGBuilder::AddCommandListSetupTask(TaskLambda&& Task)
 		ParallelSetupEvents.Emplace(UE::Tasks::Launch(TEXT("FRDGBuilder::AddCommandListSetupTask"), [Task = MoveTemp(Task), RHICmdListTask]
 		{
 			FTaskTagScope Scope(ETaskTag::EParallelRenderingThread);
+			RHICmdListTask->SwitchPipeline(ERHIPipeline::Graphics);
 			Task(*RHICmdListTask);
-
 			RHICmdListTask->FinishRecording();
-
 		}));
 
 		RHICmdList.QueueAsyncCommandListSubmit(RHICmdListTask);
diff --git a/Engine/Source/Runtime/VulkanRHI/Private/VulkanBuffer.cpp b/Engine/Source/Runtime/VulkanRHI/Private/VulkanBuffer.cpp
--- a/Engine/Source/Runtime/VulkanRHI/Private/VulkanBuffer.cpp
+++ b/Engine/Source/Runtime/VulkanRHI/Private/VulkanBuffer.cpp
@@ -559,11 +559,5 @@ void FVulkanDynamicRHI::RHITransferBufferUnderlyingResource(FRHIBuffer* DestBuff
 void FVulkanDynamicRHI::RHIUnlockBuffer(FRHICommandListBase& RHICmdList, FRHIBuffer* BufferRHI)
 {
 	QUICK_SCOPE_CYCLE_COUNTER(STAT_FDynamicRHI_UnlockBuffer_RenderThread);
-
-	// We might need to queue a copy, make sure we have an active pipeline
-	if (RHICmdList.IsTopOfPipe() && (RHICmdList.GetPipeline() == ERHIPipeline::None))
-	{
-		RHICmdList.SwitchPipeline(ERHIPipeline::Graphics);
-	}
 	FDynamicRHI::RHIUnlockBuffer(RHICmdList, BufferRHI);
 }
diff --git a/Engine/Source/Runtime/VulkanRHI/Private/VulkanTexture.cpp b/Engine/Source/Runtime/VulkanRHI/Private/VulkanTexture.cpp
--- a/Engine/Source/Runtime/VulkanRHI/Private/VulkanTexture.cpp
+++ b/Engine/Source/Runtime/VulkanRHI/Private/VulkanTexture.cpp
@@ -1233,7 +1233,7 @@ void FVulkanDynamicRHI::RHIUnlockTexture2DArray(FRHITexture2DArray* TextureRHI,
 	}
 }
 
-void FVulkanDynamicRHI::InternalUpdateTexture2D(bool bFromRenderingThread, FRHITexture2D* TextureRHI, uint32 MipIndex, const struct FUpdateTextureRegion2D& UpdateRegion, uint32 SourceRowPitch, const uint8* SourceData)
+void FVulkanDynamicRHI::InternalUpdateTexture2D(FRHICommandListBase& RHICmdList, FRHITexture2D* TextureRHI, uint32 MipIndex, const struct FUpdateTextureRegion2D& UpdateRegion, uint32 SourceRowPitch, const uint8* SourceData)
 {
 	LLM_SCOPE_VULKAN(ELLMTagVulkan::VulkanTextures);
 	FVulkanTexture* Texture = FVulkanTexture::Cast(TextureRHI);
@@ -1247,7 +1247,6 @@ void FVulkanDynamicRHI::InternalUpdateTexture2D(bool bFromRenderingThread, FRHIT
 
 	ensure(BlockSizeZ == 1);
 
-	FVulkanCommandListContext& Context = Device->GetImmediateContext();
 	const VkPhysicalDeviceLimits& Limits = Device->GetLimits();
 
 	VkBufferImageCopy Region;
@@ -1288,19 +1287,17 @@ void FVulkanDynamicRHI::InternalUpdateTexture2D(bool bFromRenderingThread, FRHIT
 	Region.imageExtent.height = UpdateRegion.Height;
 	Region.imageExtent.depth = 1;
 
-	FRHICommandList& RHICmdList = FRHICommandListExecutor::GetImmediateCommandList();
-	if (!bFromRenderingThread || (RHICmdList.Bypass() || !IsRunningRHIInSeparateThread()))
+	if (RHICmdList.IsBottomOfPipe())
 	{
-		FVulkanTexture::InternalLockWrite(Device->GetImmediateContext(), Texture, Region, StagingBuffer);
+		FVulkanTexture::InternalLockWrite(FVulkanCommandListContext::GetVulkanContext(RHICmdList.GetContext()), Texture, Region, StagingBuffer);
 	}
 	else
 	{
-		check(IsInRenderingThread());
 		ALLOC_COMMAND_CL(RHICmdList, FRHICommandLockWriteTexture)(Texture, Region, StagingBuffer);
 	}
 }
 
-FUpdateTexture3DData FVulkanDynamicRHI::BeginUpdateTexture3D_RenderThread(class FRHICommandListImmediate& RHICmdList, FRHITexture3D* Texture, uint32 MipIndex, const struct FUpdateTextureRegion3D& UpdateRegion)
+FUpdateTexture3DData FVulkanDynamicRHI::RHIBeginUpdateTexture3D(FRHICommandListBase& RHICmdList, FRHITexture3D* Texture, uint32 MipIndex, const struct FUpdateTextureRegion3D& UpdateRegion)
 {
 	const int32 FormatSize = PixelFormatBlockBytes[Texture->GetFormat()];
 	const int32 RowPitch = UpdateRegion.Width * FormatSize;
@@ -1312,18 +1309,18 @@ FUpdateTexture3DData FVulkanDynamicRHI::BeginUpdateTexture3D_RenderThread(class
 	return FUpdateTexture3DData(Texture, MipIndex, UpdateRegion, RowPitch, DepthPitch, Data, MemorySize, GFrameNumberRenderThread);
 }
 
-void FVulkanDynamicRHI::EndUpdateTexture3D_RenderThread(class FRHICommandListImmediate& RHICmdList, FUpdateTexture3DData& UpdateData)
+void FVulkanDynamicRHI::RHIEndUpdateTexture3D(FRHICommandListBase& RHICmdList, FUpdateTexture3DData& UpdateData)
 {
-	check(IsInRenderingThread());
+	check(IsInParallelRenderingThread());
 	check(GFrameNumberRenderThread == UpdateData.FrameNumber);
 
-	InternalUpdateTexture3D(true, UpdateData.Texture, UpdateData.MipIndex, UpdateData.UpdateRegion, UpdateData.RowPitch, UpdateData.DepthPitch, UpdateData.Data);
+	InternalUpdateTexture3D(RHICmdList, UpdateData.Texture, UpdateData.MipIndex, UpdateData.UpdateRegion, UpdateData.RowPitch, UpdateData.DepthPitch, UpdateData.Data);
 	
 	FMemory::Free(UpdateData.Data);
 	UpdateData.Data = nullptr;
 }
 
-void FVulkanDynamicRHI::InternalUpdateTexture3D(bool bFromRenderingThread, FRHITexture3D* TextureRHI, uint32 MipIndex, const FUpdateTextureRegion3D& UpdateRegion, uint32 SourceRowPitch, uint32 SourceDepthPitch, const uint8* SourceData)
+void FVulkanDynamicRHI::InternalUpdateTexture3D(FRHICommandListBase& RHICmdList, FRHITexture3D* TextureRHI, uint32 MipIndex, const FUpdateTextureRegion3D& UpdateRegion, uint32 SourceRowPitch, uint32 SourceDepthPitch, const uint8* SourceData)
 {
 	LLM_SCOPE_VULKAN(ELLMTagVulkan::VulkanTextures);
 	FVulkanTexture* Texture = FVulkanTexture::Cast(TextureRHI);
@@ -1337,7 +1334,6 @@ void FVulkanDynamicRHI::InternalUpdateTexture3D(bool bFromRenderingThread, FRHIT
 
 	ensure(BlockSizeZ == 1);
 
-	FVulkanCommandListContext& Context = Device->GetImmediateContext();
 	const VkPhysicalDeviceLimits& Limits = Device->GetLimits();
 
 	VkBufferImageCopy Region;
@@ -1388,14 +1384,12 @@ void FVulkanDynamicRHI::InternalUpdateTexture3D(bool bFromRenderingThread, FRHIT
 	Region.imageExtent.height = (uint32)FMath::Min((int32)(TextureSizeY-UpdateRegion.DestY), (int32)UpdateRegion.Height);
 	Region.imageExtent.depth = (uint32)FMath::Min((int32)(TextureSizeZ-UpdateRegion.DestZ), (int32)UpdateRegion.Depth);
 
-	FRHICommandList& RHICmdList = FRHICommandListExecutor::GetImmediateCommandList();
-	if (!bFromRenderingThread || (RHICmdList.Bypass() || !IsRunningRHIInSeparateThread()))
+	if (RHICmdList.IsBottomOfPipe())
 	{
-		FVulkanTexture::InternalLockWrite(Device->GetImmediateContext(), Texture, Region, StagingBuffer);
+		FVulkanTexture::InternalLockWrite(FVulkanCommandListContext::GetVulkanContext(RHICmdList.GetContext()), Texture, Region, StagingBuffer);
 	}
 	else
 	{
-		check(IsInRenderingThread());
 		ALLOC_COMMAND_CL(RHICmdList, FRHICommandLockWriteTexture)(Texture, Region, StagingBuffer);
 	}
 }
diff --git a/Engine/Source/Runtime/VulkanRHI/Private/VulkanUniformBuffer.cpp b/Engine/Source/Runtime/VulkanRHI/Private/VulkanUniformBuffer.cpp
--- a/Engine/Source/Runtime/VulkanRHI/Private/VulkanUniformBuffer.cpp
+++ b/Engine/Source/Runtime/VulkanRHI/Private/VulkanUniformBuffer.cpp
@@ -332,12 +332,6 @@ inline void FVulkanDynamicRHI::UpdateUniformBuffer(FRHICommandListBase& RHICmdLi
 
 			FVulkanCommandListContextImmediate* Context = &Device->GetImmediateContext();
 
-			// When we don't use the ring buffer, we'll need an active pipeline for a copy
-			if (!bUseRingBuffer && (RHICmdList.GetPipeline() == ERHIPipeline::None))
-			{
-				RHICmdList.SwitchPipeline(ERHIPipeline::Graphics);
-			}
-
 			RHICmdList.EnqueueLambda([UniformBuffer, CmdListResources, NumResources, ConstantBufferSize, CmdListConstantBufferData](FRHICommandListBase& CmdList)
 			{
 				FVulkanCommandListContext& Context = (FVulkanCommandListContext&)CmdList.GetContext().GetLowestLevelContext();
diff --git a/Engine/Source/Runtime/VulkanRHI/Public/VulkanDynamicRHI.h b/Engine/Source/Runtime/VulkanRHI/Public/VulkanDynamicRHI.h
--- a/Engine/Source/Runtime/VulkanRHI/Public/VulkanDynamicRHI.h
+++ b/Engine/Source/Runtime/VulkanRHI/Public/VulkanDynamicRHI.h
@@ -145,13 +145,13 @@ public:
 	}
 	virtual void* RHILockTexture2DArray(FRHITexture2DArray* Texture, uint32 TextureIndex, uint32 MipIndex, EResourceLockMode LockMode, uint32& DestStride, bool bLockWithinMiptail) final override;
 	virtual void RHIUnlockTexture2DArray(FRHITexture2DArray* Texture, uint32 TextureIndex, uint32 MipIndex, bool bLockWithinMiptail) final override;
-	virtual void RHIUpdateTexture2D(FRHITexture2D* Texture, uint32 MipIndex, const struct FUpdateTextureRegion2D& UpdateRegion, uint32 SourcePitch, const uint8* SourceData) final override
+	virtual void RHIUpdateTexture2D(FRHICommandListBase& RHICmdList, FRHITexture2D* Texture, uint32 MipIndex, const struct FUpdateTextureRegion2D& UpdateRegion, uint32 SourcePitch, const uint8* SourceData) final override
 	{
-		InternalUpdateTexture2D(false, Texture, MipIndex, UpdateRegion, SourcePitch, SourceData);
+		InternalUpdateTexture2D(RHICmdList, Texture, MipIndex, UpdateRegion, SourcePitch, SourceData);
 	}
-	virtual void RHIUpdateTexture3D(FRHITexture3D* Texture, uint32 MipIndex, const struct FUpdateTextureRegion3D& UpdateRegion, uint32 SourceRowPitch, uint32 SourceDepthPitch, const uint8* SourceData) final override
+	virtual void RHIUpdateTexture3D(FRHICommandListBase& RHICmdList, FRHITexture3D* Texture, uint32 MipIndex, const struct FUpdateTextureRegion3D& UpdateRegion, uint32 SourceRowPitch, uint32 SourceDepthPitch, const uint8* SourceData) final override
 	{
-		InternalUpdateTexture3D(false, Texture, MipIndex, UpdateRegion, SourceRowPitch, SourceDepthPitch, SourceData);
+		InternalUpdateTexture3D(RHICmdList, Texture, MipIndex, UpdateRegion, SourceRowPitch, SourceDepthPitch, SourceData);
 	}
 	virtual void* RHILockTextureCubeFace(FRHITextureCube* Texture, uint32 FaceIndex, uint32 ArrayIndex, uint32 MipIndex, EResourceLockMode LockMode, uint32& DestStride, bool bLockWithinMiptail) final override;
 	virtual void RHIUnlockTextureCubeFace(FRHITextureCube* Texture, uint32 FaceIndex, uint32 ArrayIndex, uint32 MipIndex, bool bLockWithinMiptail) final override;
@@ -231,17 +231,9 @@ public:
 	{
 		InternalUnlockTexture2D(true, Texture, MipIndex, bLockWithinMiptail);
 	}
-	virtual void UpdateTexture2D_RenderThread(class FRHICommandListImmediate& RHICmdList, FRHITexture2D* Texture, uint32 MipIndex, const struct FUpdateTextureRegion2D& UpdateRegion, uint32 SourcePitch, const uint8* SourceData) override final
-	{
-		InternalUpdateTexture2D(true, Texture, MipIndex, UpdateRegion, SourcePitch, SourceData);
-	}
-	virtual void UpdateTexture3D_RenderThread(class FRHICommandListImmediate& RHICmdList, FRHITexture3D* Texture, uint32 MipIndex, const struct FUpdateTextureRegion3D& UpdateRegion, uint32 SourceRowPitch, uint32 SourceDepthPitch, const uint8* SourceData) override final
-	{
-		InternalUpdateTexture3D(true, Texture, MipIndex, UpdateRegion, SourceRowPitch, SourceDepthPitch, SourceData);
-	}
 	
-	virtual FUpdateTexture3DData BeginUpdateTexture3D_RenderThread(class FRHICommandListImmediate& RHICmdList, FRHITexture3D* Texture, uint32 MipIndex, const struct FUpdateTextureRegion3D& UpdateRegion) override final;
-	virtual void EndUpdateTexture3D_RenderThread(class FRHICommandListImmediate& RHICmdList, FUpdateTexture3DData& UpdateData) override final;
+	virtual FUpdateTexture3DData RHIBeginUpdateTexture3D(FRHICommandListBase& RHICmdList, FRHITexture3D* Texture, uint32 MipIndex, const struct FUpdateTextureRegion3D& UpdateRegion) override final;
+	virtual void RHIEndUpdateTexture3D(FRHICommandListBase& RHICmdList, FUpdateTexture3DData& UpdateData) override final;
 
 	virtual FTextureRHIRef RHICreateTexture_RenderThread(class FRHICommandListImmediate& RHICmdList, const FRHITextureCreateDesc& CreateDesc) final override
 	{
@@ -449,8 +441,8 @@ protected:
 	FCriticalSection LockBufferCS;
 
 	void InternalUnlockTexture2D(bool bFromRenderingThread, FRHITexture2D* Texture, uint32 MipIndex, bool bLockWithinMiptail);
-	void InternalUpdateTexture2D(bool bFromRenderingThread, FRHITexture2D* TextureRHI, uint32 MipIndex, const struct FUpdateTextureRegion2D& UpdateRegion, uint32 SourcePitch, const uint8* SourceData);
-	void InternalUpdateTexture3D(bool bFromRenderingThread, FRHITexture3D* TextureRHI, uint32 MipIndex, const struct FUpdateTextureRegion3D& UpdateRegion, uint32 SourceRowPitch, uint32 SourceDepthPitch, const uint8* SourceData);
+	void InternalUpdateTexture2D(FRHICommandListBase& RHICmdList, FRHITexture2D* TextureRHI, uint32 MipIndex, const struct FUpdateTextureRegion2D& UpdateRegion, uint32 SourcePitch, const uint8* SourceData);
+	void InternalUpdateTexture3D(FRHICommandListBase& RHICmdList, FRHITexture3D* TextureRHI, uint32 MipIndex, const struct FUpdateTextureRegion3D& UpdateRegion, uint32 SourceRowPitch, uint32 SourceDepthPitch, const uint8* SourceData);
 
 	void UpdateUniformBuffer(FRHICommandListBase& RHICmdList, FVulkanUniformBuffer* UniformBuffer, const void* Contents);
 
diff --git a/Engine/Source/Runtime/Windows/D3D11RHI/Private/D3D11RHIPrivate.h b/Engine/Source/Runtime/Windows/D3D11RHI/Private/D3D11RHIPrivate.h
--- a/Engine/Source/Runtime/Windows/D3D11RHI/Private/D3D11RHIPrivate.h
+++ b/Engine/Source/Runtime/Windows/D3D11RHI/Private/D3D11RHIPrivate.h
@@ -506,8 +506,8 @@ public:
 	virtual void RHIUnlockTexture2D(FRHITexture2D* Texture, uint32 MipIndex, bool bLockWithinMiptail) final override;
 	virtual void* RHILockTexture2DArray(FRHITexture2DArray* Texture, uint32 TextureIndex, uint32 MipIndex, EResourceLockMode LockMode, uint32& DestStride, bool bLockWithinMiptail) final override;
 	virtual void RHIUnlockTexture2DArray(FRHITexture2DArray* Texture, uint32 TextureIndex, uint32 MipIndex, bool bLockWithinMiptail) final override;
-	virtual void RHIUpdateTexture2D(FRHITexture2D* Texture, uint32 MipIndex, const struct FUpdateTextureRegion2D& UpdateRegion, uint32 SourcePitch, const uint8* SourceData) final override;
-	virtual void RHIUpdateTexture3D(FRHITexture3D* Texture, uint32 MipIndex, const struct FUpdateTextureRegion3D& UpdateRegion, uint32 SourceRowPitch, uint32 SourceDepthPitch, const uint8* SourceData) final override;
+	virtual void RHIUpdateTexture2D(FRHICommandListBase& RHICmdList, FRHITexture2D* Texture, uint32 MipIndex, const struct FUpdateTextureRegion2D& UpdateRegion, uint32 SourcePitch, const uint8* SourceData) final override;
+	virtual void RHIUpdateTexture3D(FRHICommandListBase& RHICmdList, FRHITexture3D* Texture, uint32 MipIndex, const struct FUpdateTextureRegion3D& UpdateRegion, uint32 SourceRowPitch, uint32 SourceDepthPitch, const uint8* SourceData) final override;
 	virtual void* RHILockTextureCubeFace(FRHITextureCube* Texture, uint32 FaceIndex, uint32 ArrayIndex, uint32 MipIndex, EResourceLockMode LockMode, uint32& DestStride, bool bLockWithinMiptail) final override;
 	virtual void RHIUnlockTextureCubeFace(FRHITextureCube* Texture, uint32 FaceIndex, uint32 ArrayIndex, uint32 MipIndex, bool bLockWithinMiptail) final override;
 	virtual void RHIBindDebugLabelName(FRHITexture* Texture, const TCHAR* Name) final override;
@@ -643,9 +643,7 @@ public:
 	virtual FComputeShaderRHIRef CreateComputeShader_RenderThread(class FRHICommandListImmediate& RHICmdList, TArrayView<const uint8> Code, const FSHAHash& Hash) final override;
 	virtual void* LockTexture2D_RenderThread(class FRHICommandListImmediate& RHICmdList, FRHITexture2D* Texture, uint32 MipIndex, EResourceLockMode LockMode, uint32& DestStride, bool bLockWithinMiptail, bool bNeedsDefaultRHIFlush = true) final override;
 	virtual void UnlockTexture2D_RenderThread(class FRHICommandListImmediate& RHICmdList, FRHITexture2D* Texture, uint32 MipIndex, bool bLockWithinMiptail, bool bNeedsDefaultRHIFlush = true) final override;
-	virtual void UpdateTexture2D_RenderThread(class FRHICommandListImmediate& RHICmdList, FRHITexture2D* Texture, uint32 MipIndex, const struct FUpdateTextureRegion2D& UpdateRegion, uint32 SourcePitch, const uint8* SourceData) final override;
-	virtual void EndUpdateTexture3D_RenderThread(class FRHICommandListImmediate& RHICmdList, FUpdateTexture3DData& UpdateData) final override;
-	virtual void UpdateTexture3D_RenderThread(class FRHICommandListImmediate& RHICmdList, FRHITexture3D* Texture, uint32 MipIndex, const struct FUpdateTextureRegion3D& UpdateRegion, uint32 SourceRowPitch, uint32 SourceDepthPitch, const uint8* SourceData) final override;
+	virtual void RHIEndUpdateTexture3D(FRHICommandListBase& RHICmdList, FUpdateTexture3DData& UpdateData) final override;
 	virtual FUnorderedAccessViewRHIRef RHICreateUnorderedAccessView_RenderThread(class FRHICommandListImmediate& RHICmdList, FRHIBuffer* Buffer, bool bUseUAVCounter, bool bAppendBuffer) final override;
 	virtual FUnorderedAccessViewRHIRef RHICreateUnorderedAccessView_RenderThread(class FRHICommandListImmediate& RHICmdList, FRHITexture* Texture, uint32 MipLevel, uint16 FirstArraySlice, uint16 NumArraySlices) final override;
 	virtual FUnorderedAccessViewRHIRef RHICreateUnorderedAccessView_RenderThread(class FRHICommandListImmediate& RHICmdList, FRHIBuffer* Buffer, uint8 Format) final override;
diff --git a/Engine/Source/Runtime/Windows/D3D11RHI/Private/D3D11Texture.cpp b/Engine/Source/Runtime/Windows/D3D11RHI/Private/D3D11Texture.cpp
--- a/Engine/Source/Runtime/Windows/D3D11RHI/Private/D3D11Texture.cpp
+++ b/Engine/Source/Runtime/Windows/D3D11RHI/Private/D3D11Texture.cpp
@@ -1514,125 +1514,81 @@ void FD3D11DynamicRHI::RHIUnlockTexture2DArray(FRHITexture2DArray* TextureRHI, u
 	Texture->Unlock(this, MipIndex, TextureIndex);
 }
 
-void FD3D11DynamicRHI::RHIUpdateTexture2D(FRHITexture2D* TextureRHI, uint32 MipIndex, const FUpdateTextureRegion2D& UpdateRegion, uint32 SourcePitch, const uint8* SourceData)
+void FD3D11DynamicRHI::RHIUpdateTexture2D(FRHICommandListBase& RHICmdList, FRHITexture2D* TextureRHI, uint32 MipIndex, const FUpdateTextureRegion2D& UpdateRegion, uint32 SourcePitch, const uint8* SourceData)
 {
-	FD3D11Texture* Texture = ResourceCast(TextureRHI);
+	bool bFreeMemory = !(RHICmdList.Bypass() || !IsRunningRHIInSeparateThread());
 
-	D3D11_BOX DestBox =
+	if (bFreeMemory)
 	{
-		UpdateRegion.DestX,                      UpdateRegion.DestY,                       0,
-		UpdateRegion.DestX + UpdateRegion.Width, UpdateRegion.DestY + UpdateRegion.Height, 1
-	};
-
-	check(UpdateRegion.Width % GPixelFormats[Texture->GetFormat()].BlockSizeX == 0);
-	check(UpdateRegion.Height % GPixelFormats[Texture->GetFormat()].BlockSizeY == 0);
-	check(UpdateRegion.DestX % GPixelFormats[Texture->GetFormat()].BlockSizeX == 0);
-	check(UpdateRegion.DestY % GPixelFormats[Texture->GetFormat()].BlockSizeY == 0);
-	check(UpdateRegion.SrcX % GPixelFormats[Texture->GetFormat()].BlockSizeX == 0);
-	check(UpdateRegion.SrcY % GPixelFormats[Texture->GetFormat()].BlockSizeY == 0);
-
-	Direct3DDeviceIMContext->UpdateSubresource(Texture->GetResource(), MipIndex, &DestBox, SourceData, SourcePitch, 0);
-}
-
-void FD3D11DynamicRHI::UpdateTexture2D_RenderThread(
-	class FRHICommandListImmediate& RHICmdList,
-	FRHITexture2D* Texture,
-	uint32 MipIndex,
-	const struct FUpdateTextureRegion2D& UpdateRegion,
-	uint32 SourcePitch,
-	const uint8* SourceData)
-{
-	if (ShouldNotEnqueueRHICommand())
-	{
-		RHIUpdateTexture2D(Texture, MipIndex, UpdateRegion, SourcePitch, SourceData);
-	}
-	else
-	{
-		const FPixelFormatInfo& FormatInfo = GPixelFormats[Texture->GetFormat()];
+		const FPixelFormatInfo& FormatInfo = GPixelFormats[TextureRHI->GetFormat()];
 		const size_t UpdateHeightInTiles = FMath::DivideAndRoundUp(UpdateRegion.Height, (uint32)FormatInfo.BlockSizeY);
 		const size_t SourceDataSize = static_cast<size_t>(SourcePitch) * UpdateHeightInTiles;
 		uint8* SourceDataCopy = (uint8*)FMemory::Malloc(SourceDataSize);
 		FMemory::Memcpy(SourceDataCopy, SourceData, SourceDataSize);
-		RunOnRHIThread([this, Texture, MipIndex, UpdateRegion, SourcePitch, SourceDataCopy]()
-		{
-			RHIUpdateTexture2D(Texture, MipIndex, UpdateRegion, SourcePitch, SourceDataCopy);
-			FMemory::Free(SourceDataCopy);
-		});
+		SourceData = SourceDataCopy;
 	}
-}
 
-void FD3D11DynamicRHI::RHIUpdateTexture3D(FRHITexture3D* TextureRHI,uint32 MipIndex,const FUpdateTextureRegion3D& UpdateRegion,uint32 SourceRowPitch,uint32 SourceDepthPitch,const uint8* SourceData)
-{
-	FD3D11Texture* Texture = ResourceCast(TextureRHI);
+	RHICmdList.EnqueueLambda([this, TextureRHI, MipIndex, UpdateRegion, SourcePitch, SourceData, bFreeMemory] (FRHICommandListBase&)
+	{
+		FD3D11Texture* Texture = ResourceCast(TextureRHI);
+
+		D3D11_BOX DestBox =
+		{
+			UpdateRegion.DestX,                      UpdateRegion.DestY,                       0,
+			UpdateRegion.DestX + UpdateRegion.Width, UpdateRegion.DestY + UpdateRegion.Height, 1
+		};
 
-	// The engine calls this with the texture size in the region. 
-	// Some platforms like D3D11 needs that to be rounded up to the block size.
-	const FPixelFormatInfo& Format = GPixelFormats[Texture->GetFormat()];
-	const int32 NumBlockX = FMath::DivideAndRoundUp<int32>(UpdateRegion.Width, Format.BlockSizeX);
-	const int32 NumBlockY = FMath::DivideAndRoundUp<int32>(UpdateRegion.Height, Format.BlockSizeY);
+		check(UpdateRegion.Width % GPixelFormats[Texture->GetFormat()].BlockSizeX == 0);
+		check(UpdateRegion.Height % GPixelFormats[Texture->GetFormat()].BlockSizeY == 0);
+		check(UpdateRegion.DestX % GPixelFormats[Texture->GetFormat()].BlockSizeX == 0);
+		check(UpdateRegion.DestY % GPixelFormats[Texture->GetFormat()].BlockSizeY == 0);
+		check(UpdateRegion.SrcX % GPixelFormats[Texture->GetFormat()].BlockSizeX == 0);
+		check(UpdateRegion.SrcY % GPixelFormats[Texture->GetFormat()].BlockSizeY == 0);
 
-	D3D11_BOX DestBox =
-	{
-		UpdateRegion.DestX,                                 UpdateRegion.DestY,                                 UpdateRegion.DestZ,
-		UpdateRegion.DestX + NumBlockX * Format.BlockSizeX, UpdateRegion.DestY + NumBlockY * Format.BlockSizeY, UpdateRegion.DestZ + UpdateRegion.Depth
-	};
+		Direct3DDeviceIMContext->UpdateSubresource(Texture->GetResource(), MipIndex, &DestBox, SourceData, SourcePitch, 0);
 
-	Direct3DDeviceIMContext->UpdateSubresource(Texture->GetResource(), MipIndex, &DestBox, SourceData, SourceRowPitch, SourceDepthPitch);
+		if (bFreeMemory) { FMemory::Free((void*)SourceData); }
+	});
 }
 
-void FD3D11DynamicRHI::EndUpdateTexture3D_RenderThread(class FRHICommandListImmediate& RHICmdList, FUpdateTexture3DData& UpdateData)
+void FD3D11DynamicRHI::RHIUpdateTexture3D(FRHICommandListBase& RHICmdList, FRHITexture3D* TextureRHI,uint32 MipIndex,const FUpdateTextureRegion3D& UpdateRegion,uint32 SourceRowPitch,uint32 SourceDepthPitch,const uint8* SourceData)
 {
-	if (RHICmdList.Bypass())
+	bool bFreeMemory = !(RHICmdList.Bypass() || !IsRunningRHIInSeparateThread());
+	if (bFreeMemory)
 	{
-		RHIUpdateTexture3D(UpdateData.Texture, UpdateData.MipIndex, UpdateData.UpdateRegion, UpdateData.RowPitch, UpdateData.DepthPitch, UpdateData.Data);
-		FMemory::Free(UpdateData.Data);
+		const SIZE_T SourceDataSize = static_cast<SIZE_T>(SourceDepthPitch) * UpdateRegion.Depth;
+		uint8* SourceDataCopy = (uint8*)FMemory::Malloc(SourceDataSize);
+		FMemory::Memcpy(SourceDataCopy, SourceData, SourceDataSize);
+		SourceData = SourceDataCopy;
 	}
-	else
+
+	RHICmdList.EnqueueLambda([this, TextureRHI, MipIndex, UpdateRegion, SourceRowPitch, SourceDepthPitch, SourceData, bFreeMemory] (FRHICommandListBase&)
 	{
-		UpdateData.Texture->AddRef();
-		RunOnRHIThread(
-			[UpdateData]()
+		FD3D11Texture* Texture = ResourceCast(TextureRHI);
+
+		// The engine calls this with the texture size in the region. 
+		// Some platforms like D3D11 needs that to be rounded up to the block size.
+		const FPixelFormatInfo& Format = GPixelFormats[Texture->GetFormat()];
+		const int32 NumBlockX = FMath::DivideAndRoundUp<int32>(UpdateRegion.Width, Format.BlockSizeX);
+		const int32 NumBlockY = FMath::DivideAndRoundUp<int32>(UpdateRegion.Height, Format.BlockSizeY);
+
+		D3D11_BOX DestBox =
 		{
-			GD3D11RHI->RHIUpdateTexture3D(
-				UpdateData.Texture,
-				UpdateData.MipIndex,
-				UpdateData.UpdateRegion,
-				UpdateData.RowPitch,
-				UpdateData.DepthPitch,
-				UpdateData.Data);
-			UpdateData.Texture->Release();
-			FMemory::Free(UpdateData.Data);
-		});
-		RHICmdList.RHIThreadFence(true);
-	}
+			UpdateRegion.DestX,                                 UpdateRegion.DestY,                                 UpdateRegion.DestZ,
+			UpdateRegion.DestX + NumBlockX * Format.BlockSizeX, UpdateRegion.DestY + NumBlockY * Format.BlockSizeY, UpdateRegion.DestZ + UpdateRegion.Depth
+		};
 
-	UpdateData.Data = nullptr;
+		Direct3DDeviceIMContext->UpdateSubresource(Texture->GetResource(), MipIndex, &DestBox, SourceData, SourceRowPitch, SourceDepthPitch);
+
+		if (bFreeMemory) { FMemory::Free((void*)SourceData); }
+	});
 }
 
-void FD3D11DynamicRHI::UpdateTexture3D_RenderThread(
-	class FRHICommandListImmediate& RHICmdList,
-	FRHITexture3D* Texture,
-	uint32 MipIndex,
-	const struct FUpdateTextureRegion3D& UpdateRegion,
-	uint32 SourceRowPitch,
-	uint32 SourceDepthPitch,
-	const uint8* SourceData)
+void FD3D11DynamicRHI::RHIEndUpdateTexture3D(FRHICommandListBase& RHICmdList, FUpdateTexture3DData& UpdateData)
 {
-	if (ShouldNotEnqueueRHICommand())
-	{
-		RHIUpdateTexture3D(Texture, MipIndex, UpdateRegion, SourceRowPitch, SourceDepthPitch, SourceData);
-	}
-	else
-	{
-		const SIZE_T SourceDataSize = static_cast<SIZE_T>(SourceDepthPitch) * UpdateRegion.Depth;
-		uint8* SourceDataCopy = (uint8*)FMemory::Malloc(SourceDataSize);
-		FMemory::Memcpy(SourceDataCopy, SourceData, SourceDataSize);
-		RunOnRHIThread([this, Texture, MipIndex, UpdateRegion, SourceRowPitch, SourceDepthPitch, SourceDataCopy]()
-		{
-			RHIUpdateTexture3D(Texture, MipIndex, UpdateRegion, SourceRowPitch, SourceDepthPitch, SourceDataCopy);
-			FMemory::Free(SourceDataCopy);
-		});
-	}
+	RHIUpdateTexture3D(RHICmdList, UpdateData.Texture, UpdateData.MipIndex, UpdateData.UpdateRegion, UpdateData.RowPitch, UpdateData.DepthPitch, UpdateData.Data);
+	FMemory::Free(UpdateData.Data);
+	UpdateData.Data = nullptr;
 }
 
 /*-----------------------------------------------------------------------------
