Add test for RHICopyTexture between UNorm and sRGB formats.

diff --git a/Engine/Plugins/Tests/RHITests/Source/RHITests/Public/RHITextureTests.h b/Engine/Plugins/Tests/RHITests/Source/RHITests/Public/RHITextureTests.h
--- a/Engine/Plugins/Tests/RHITests/Source/RHITests/Public/RHITextureTests.h
+++ b/Engine/Plugins/Tests/RHITests/Source/RHITests/Public/RHITextureTests.h
@@ -941,6 +941,50 @@ public:
 				}));
 			}
 
+			// Copy textures with UNorm or sRGB settings (Dest should remain bitwise identical to Source)
+			const auto CopyTextureSrgbLambda = [&](const TCHAR* TestName, bool bSourceSrgb, bool bDestSrgb)
+			{
+				const FIntPoint Extent(16, 16);
+				const FIntPoint PatternOffset(0, 0);
+				const FIntPoint PatternSize(16, 16);
+				const int32 ImageIndex = 1;
+
+				const FRHITextureCreateDesc SourceDesc =
+					FRHITextureCreateDesc::Create2D(TestName)
+					.SetExtent(Extent.X, Extent.Y)
+					.SetFormat(PF_B8G8R8A8)
+					.SetFlags(bSourceSrgb ? ETextureCreateFlags::SRGB : ETextureCreateFlags::None)
+					.SetInitialState(ERHIAccess::CopySrc);
+
+				FTextureRHIRef SourceTexture = RHICreateTexture(SourceDesc);
+
+				uint32 Stride = 0;
+				void* Data = RHICmdList.LockTexture2D(SourceTexture, 0, RLM_WriteOnly, Stride, false);
+				WriteTestData(Data, Stride, PatternSize, PatternOffset, ImageIndex);
+				RHICmdList.UnlockTexture2D(SourceTexture, 0, false);
+
+				const FRHITextureCreateDesc DestDesc =
+					FRHITextureCreateDesc::Create2D(TestName)
+					.SetExtent(Extent.X, Extent.Y)
+					.SetFormat(PF_B8G8R8A8)
+					.SetFlags(bDestSrgb ? ETextureCreateFlags::SRGB : ETextureCreateFlags::None)
+					.SetInitialState(ERHIAccess::CopyDest);
+
+				FTextureRHIRef DestTexture = RHICreateTexture(DestDesc);
+				RHICmdList.CopyTexture(SourceTexture, DestTexture, {});
+				RHICmdList.Transition(FRHITransitionInfo(DestTexture, ERHIAccess::CopyDest, ERHIAccess::CopySrc));
+
+				RUN_TEST(VerifyTextureContents(TestName, RHICmdList, DestTexture, [&](void* Data, uint32 MipWidth, uint32 MipHeight, uint32 Width, uint32 Height, uint32 MipIndex, uint32 SliceIndex)
+				{
+					return CheckTestData(Data, Width * sizeof(uint32), PatternSize, PatternOffset, ImageIndex);
+				}));
+			};
+
+			CopyTextureSrgbLambda(TEXT("Copy Texture UNorm to UNorm"), false, false);
+			CopyTextureSrgbLambda(TEXT("Copy Texture UNorm to sRGB"), false, true);
+			CopyTextureSrgbLambda(TEXT("Copy Texture sRGB to UNorm"), true, false);
+			CopyTextureSrgbLambda(TEXT("Copy Texture sRGB to sRGB"), true, true);
+
 			// Copy mip to another texture.
 			const auto CopyTextureMipsLambda = [&](bool bUseExplicitSize)
 			{
