Fix RHI unit tests:
* wrong souce states for a couple of transitions
* wrong buffer size in the TextureFromBuffer update test (causing a GPU page fault due to the source buffer being too small)
* don't try to run the parallel buffer creation test when bypass mode is enabled, since parallel command lists are not supported in that mode

#rnx

diff --git a/Engine/Plugins/Tests/RHITests/Source/RHITests/Public/RHIBufferTests.h b/Engine/Plugins/Tests/RHITests/Source/RHITests/Public/RHIBufferTests.h
--- a/Engine/Plugins/Tests/RHITests/Source/RHITests/Public/RHIBufferTests.h
+++ b/Engine/Plugins/Tests/RHITests/Source/RHITests/Public/RHIBufferTests.h
@@ -352,7 +352,7 @@ public:
 
 	static bool Test_RHICreateBuffer_Parallel(FRHICommandListImmediate& RHICmdList)
 	{
-		if (!GRHISupportsMultithreadedResources)
+		if (!GRHISupportsMultithreadedResources || RHICmdList.Bypass())
 		{
 			return true;
 		}
diff --git a/Engine/Plugins/Tests/RHITests/Source/RHITests/Public/RHITextureTests.h b/Engine/Plugins/Tests/RHITests/Source/RHITests/Public/RHITextureTests.h
--- a/Engine/Plugins/Tests/RHITests/Source/RHITests/Public/RHITextureTests.h
+++ b/Engine/Plugins/Tests/RHITests/Source/RHITests/Public/RHITextureTests.h
@@ -10,7 +10,9 @@
 
 class FRHITextureTests
 {
+	// This expects Texture to be in the CopySrc state.
 	static bool VerifyTextureContents(const TCHAR* TestName, FRHICommandListImmediate& RHICmdList, FRHITexture* Texture, TFunctionRef<bool(void* Ptr, uint32 MipWidth, uint32 MipHeight, uint32 Width, uint32 Height, uint32 MipIndex, uint32 SliceIndex)> VerifyCallback);
+
 public:
 	template <typename ValueType, uint32 NumTestBytes>
 	static bool RunTest_UAVClear_Texture(FRHICommandListImmediate& RHICmdList, const FString& TestName, FRHITexture* TextureRHI, uint32 MipIndex, const ValueType& ClearValue, void(FRHIComputeCommandList::* ClearPtr)(FRHIUnorderedAccessView*, ValueType const&), const uint8(&TestValue)[NumTestBytes])
@@ -28,7 +30,6 @@ public:
 				FMemory::Memset(&ZerosValue, 0, sizeof(ZerosValue));
 				(RHICmdList.*ClearPtr)(MipUAV, ZerosValue);
 			}
-			RHICmdList.Transition(FRHITransitionInfo(TextureRHI, ERHIAccess::UAVCompute, ERHIAccess::CopySrc));
 
 			auto VerifyMip = [&](void* Ptr, uint32 MipWidth, uint32 MipHeight, uint32 Width, uint32 Height, uint32 CurrentMipIndex, uint32 CurrentSliceIndex, bool bShouldBeZero)
 			{
@@ -68,12 +69,16 @@ public:
 			{
 				return VerifyMip(Ptr, MipWidth, MipHeight, Width, Height, CurrentMipIndex, CurrentSliceIndex, true);
 			};
+
+			RHICmdList.Transition(FRHITransitionInfo(TextureRHI, ERHIAccess::UAVCompute, ERHIAccess::CopySrc));
 			bResult0 = VerifyTextureContents(*FString::Printf(TEXT("%s - clear whole resource to zero"), *TestName), RHICmdList, TextureRHI, VerifyMipIsZero);
 
 			// Clear the selected mip index to the provided value
-			RHICmdList.Transition(FRHITransitionInfo(TextureRHI, ERHIAccess::Unknown, ERHIAccess::UAVCompute));
+			
+			RHICmdList.Transition(FRHITransitionInfo(TextureRHI, ERHIAccess::CopySrc, ERHIAccess::UAVCompute));
 			FUnorderedAccessViewRHIRef SpecificMipUAV = RHICreateUnorderedAccessView(TextureRHI, MipIndex);
 			(RHICmdList.*ClearPtr)(SpecificMipUAV, ClearValue);
+
 			RHICmdList.Transition(FRHITransitionInfo(TextureRHI, ERHIAccess::UAVCompute, ERHIAccess::CopySrc));
 			bResult1 = VerifyTextureContents(*FString::Printf(TEXT("%s - clear mip %d to (%s)"), *TestName, MipIndex, *ClearValueToString(ClearValue)), RHICmdList, TextureRHI,
 				[&](void* Ptr, uint32 MipWidth, uint32 MipHeight, uint32 Width, uint32 Height, uint32 CurrentMipIndex, uint32 CurrentSliceIndex)
@@ -454,7 +459,7 @@ public:
 	}
 
 	template <typename SourceType>
-	static bool Test_UpdateTexture_Impl(FRHICommandListImmediate& RHICmdList, FString TestName, FRHITexture* Texture, const FUpdateTextureRegion2D& Region, uint32 SourcePitch, const uint8* SourceData, const uint8* ZeroData, ETextureUpdateType UpdateType)
+	static bool Test_UpdateTexture_Impl(FRHICommandListImmediate& RHICmdList, FString TestName, FRHITexture* Texture, const FUpdateTextureRegion2D& Region, uint32 SourceSize, uint32 SourcePitch, const uint8* SourceData, const uint8* ZeroData, ETextureUpdateType UpdateType)
 	{
 		bool bResult = true;
 		
@@ -522,12 +527,10 @@ public:
 		}
 		else if (UpdateType == ETextureUpdateType::TextureFromBuffer)
 		{
-			const int32 BufferSize = Region.Width * Region.Height;
-
 			FRHIResourceCreateInfo CreateInfo(TEXT("SourceUpdateBuffer"));
-			TRefCountPtr<FRHIBuffer> Buffer = RHICmdList.CreateBuffer(BufferSize, EBufferUsageFlags::ShaderResource | EBufferUsageFlags::Static | EBufferUsageFlags::KeepCPUAccessible, SourceType::ElementSize, ERHIAccess::SRVCompute, CreateInfo);
-			void* Data = RHICmdList.LockBuffer(Buffer, 0, BufferSize, RLM_WriteOnly);
-			FMemory::Memcpy(Data, SourceData, BufferSize);
+			TRefCountPtr<FRHIBuffer> Buffer = RHICmdList.CreateBuffer(SourceSize, EBufferUsageFlags::ShaderResource | EBufferUsageFlags::Static | EBufferUsageFlags::KeepCPUAccessible, SourceType::ElementSize, ERHIAccess::SRVCompute, CreateInfo);
+			void* Data = RHICmdList.LockBuffer(Buffer, 0, SourceSize, RLM_WriteOnly);
+			FMemory::Memcpy(Data, SourceData, SourceSize);
 			RHICmdList.UnlockBuffer(Buffer);
 
 			RHICmdList.UpdateFromBufferTexture2D(Texture, 0, Region, SourcePitch, Buffer, 0);
@@ -543,7 +546,9 @@ public:
 			}
 		}
 
-		RHICmdList.Transition(FRHITransitionInfo(Texture, ERHIAccess::UAVCompute, ERHIAccess::CopySrc));
+		// FIXME: we should define the state expected by the UpdateTexture calls, and how to properly synchronize subsequent access to the texture. Until then,
+		// we'll insert a transition from Unknown here and hope for the best.
+		RHICmdList.Transition(FRHITransitionInfo(Texture, ERHIAccess::Unknown, ERHIAccess::CopySrc));
 
 		bResult &= VerifyTextureContents(*FString::Printf(TEXT("%s - update (%d,%d -> %d,%d)"), *TestName, Region.SrcX, Region.SrcY, Region.DestX, Region.DestY), RHICmdList, Texture,
 			[&](void* Ptr, uint32 MipWidth, uint32 MipHeight, uint32 Width, uint32 Height, uint32 CurrentMipIndex, uint32 CurrentSliceIndex)
@@ -637,7 +642,7 @@ public:
 				for (uint32 DestColumn = 0; DestColumn < 3; DestColumn++)
 				{
 					Region.DestX = AxisSlotToCoord(DestColumn, TextureWidth, UpdateWidth);
-					bResult &= Test_UpdateTexture_Impl<SourceType>(RHICmdList, TestName, Texture, Region, UpdateDataPitch, UpdateData, ZeroData, UpdateType);
+					bResult &= Test_UpdateTexture_Impl<SourceType>(RHICmdList, TestName, Texture, Region, SrcDataSize, UpdateDataPitch, UpdateData, ZeroData, UpdateType);
 				}
 			}
 
@@ -649,7 +654,7 @@ public:
 				for (uint32 SrcColumn = 0; SrcColumn < 3; SrcColumn++)
 				{
 					Region.SrcX = AxisSlotToCoord(SrcColumn, SrcDataWidth, UpdateWidth);
-					bResult &= Test_UpdateTexture_Impl<SourceType>(RHICmdList, TestName, Texture, Region, UpdateDataPitch, UpdateData, ZeroData, UpdateType);
+					bResult &= Test_UpdateTexture_Impl<SourceType>(RHICmdList, TestName, Texture, Region, SrcDataSize, UpdateDataPitch, UpdateData, ZeroData, UpdateType);
 				}
 			}
 #endif
