Fix FRHIGPUTextureReadback returning incorrect results when reading back a subrect of the source texture.
* readback is supposed to return only the subrect (when one is given), but we were allocating a staging buffer the same size as the source texture, and trying to copy the data to the same position it had in the source. Fixed that by using the correct size and a zero offset for the destination buffer.
* the GPU fence was written before the transition to CPURead was done, which could cause trouble in some RHIs (MapStagingSurface has ambiguous semantics, we should fix that).

Also added an RHI unit test for this, and modified the WriteTestData/CheckTestData helper functions to support different source and destination offsets and extents.

diff --git a/Engine/Plugins/Tests/RHITests/Source/RHITests/Public/RHITextureTests.h b/Engine/Plugins/Tests/RHITests/Source/RHITests/Public/RHITextureTests.h
--- a/Engine/Plugins/Tests/RHITests/Source/RHITests/Public/RHITextureTests.h
+++ b/Engine/Plugins/Tests/RHITests/Source/RHITests/Public/RHITextureTests.h
@@ -5,6 +5,7 @@
 #include "RHITestsCommon.h"
 #include <type_traits>
 #include "Math/PackedVector.h"
+#include "RHIGPUReadback.h"
 
 //PRAGMA_DISABLE_OPTIMIZATION
 
@@ -841,7 +842,7 @@ public:
 
 		const auto WriteTestData = [](void* Ptr, int32 Stride, FIntPoint Size, FIntPoint Offset, int32 ImageIndex)
 		{
-			int32 Index = ImageIndex;
+			int32 Index = ImageIndex + Offset.Y * Size.X + Offset.X;
 			uint8* Bytes = (uint8*)Ptr;
 
 			for (int32 Y = Offset.Y; Y < Offset.Y + Size.Y; Y++)
@@ -868,7 +869,7 @@ public:
 
 				for (int32 X = Offset.X; X < Offset.X + Size.X; X++)
 				{
-					DataString += FString::Printf(TEXT(" %d"), Row[X]);
+					DataString += FString::Printf(TEXT(" %4d"), Row[X]);
 				}
 
 				DataString += TEXT("\n");
@@ -877,16 +878,16 @@ public:
 			return DataString;
 		};
 
-		const auto CheckTestData = [&](const void* Ptr, int32 Stride, FIntPoint Size, FIntPoint  Offset, int32 ImageIndex) -> bool
+		const auto CheckTestData = [&](const void* Ptr, int32 Stride, FIntPoint CheckSize, FIntPoint DataOffset, FIntPoint SrcSize, FIntPoint SrcOffset, int32 ImageIndex) -> bool
 		{
-			int32 Index = ImageIndex;
 			const uint8* Bytes = (uint8*)Ptr;
+			int32 Index = ImageIndex + SrcOffset.Y * SrcSize.X + SrcOffset.X;
 
-			for (int32 Y = Offset.Y; Y < Offset.Y + Size.Y; Y++)
+			for (int32 Y = DataOffset.Y; Y < DataOffset.Y + CheckSize.Y; Y++)
 			{
 				const uint32* Row = (const uint32*)(Bytes + Y * Stride);
 
-				for (int32 X = Offset.X; X < Offset.X + Size.X; X++)
+				for (int32 X = DataOffset.X; X < DataOffset.X + CheckSize.X; X++, Index++)
 				{
 					if (Row[X] != Index)
 					{
@@ -896,13 +897,13 @@ public:
 							TEXT("\tExpected Value: %d\n")
 							TEXT("\tActual Value: %d\n")
 							TEXT("%s"),
-							X, Y, Index, Row[X], *DumpTestData(Ptr, Stride, Size, Offset));
+							X, Y, Index, Row[X], *DumpTestData(Ptr, Stride, CheckSize, DataOffset));
 
 						return false;
 					}
-
-					Index++;
 				}
+
+				Index += SrcSize.X - CheckSize.X;
 			}
 
 			return true;
@@ -913,8 +914,6 @@ public:
 			{
 				const TCHAR* TestName = TEXT("Copy Texture Whole");
 				const FIntPoint Extent(16, 16);
-				const FIntPoint PatternOffset(0, 0);
-				const FIntPoint PatternSize(16, 16);
 				const int32 ImageIndex = 1;
 
 				const FRHITextureCreateDesc SourceDesc =
@@ -927,7 +926,7 @@ public:
 
 				uint32 Stride = 0;
 				void* Data = RHICmdList.LockTexture2D(SourceTexture, 0, RLM_WriteOnly, Stride, false);
-				WriteTestData(Data, Stride, PatternSize, PatternOffset, ImageIndex);
+				WriteTestData(Data, Stride, Extent, FIntPoint::ZeroValue, ImageIndex);
 				RHICmdList.UnlockTexture2D(SourceTexture, 0, false);
 
 				const FRHITextureCreateDesc DestDesc =
@@ -942,7 +941,7 @@ public:
 
 				RUN_TEST(VerifyTextureContents(TestName, RHICmdList, DestTexture, [&](void* Data, uint32 MipWidth, uint32 MipHeight, uint32 Width, uint32 Height, uint32 MipIndex, uint32 SliceIndex)
 				{
-					return CheckTestData(Data, Width * sizeof(uint32), PatternSize, PatternOffset, ImageIndex);
+					return CheckTestData(Data, Width * sizeof(uint32), Extent, FIntPoint::ZeroValue, Extent, FIntPoint::ZeroValue, ImageIndex);
 				}));
 			}
 
@@ -950,8 +949,6 @@ public:
 			const auto CopyTextureSrgbLambda = [&](const TCHAR* TestName, bool bSourceSrgb, bool bDestSrgb)
 			{
 				const FIntPoint Extent(16, 16);
-				const FIntPoint PatternOffset(0, 0);
-				const FIntPoint PatternSize(16, 16);
 				const int32 ImageIndex = 1;
 
 				const FRHITextureCreateDesc SourceDesc =
@@ -965,7 +962,7 @@ public:
 
 				uint32 Stride = 0;
 				void* Data = RHICmdList.LockTexture2D(SourceTexture, 0, RLM_WriteOnly, Stride, false);
-				WriteTestData(Data, Stride, PatternSize, PatternOffset, ImageIndex);
+				WriteTestData(Data, Stride, Extent, FIntPoint::ZeroValue, ImageIndex);
 				RHICmdList.UnlockTexture2D(SourceTexture, 0, false);
 
 				const FRHITextureCreateDesc DestDesc =
@@ -981,7 +978,7 @@ public:
 
 				RUN_TEST(VerifyTextureContents(TestName, RHICmdList, DestTexture, [&](void* Data, uint32 MipWidth, uint32 MipHeight, uint32 Width, uint32 Height, uint32 MipIndex, uint32 SliceIndex)
 				{
-					return CheckTestData(Data, Width * sizeof(uint32), PatternSize, PatternOffset, ImageIndex);
+					return CheckTestData(Data, Width * sizeof(uint32), Extent, FIntPoint::ZeroValue, Extent, FIntPoint::ZeroValue, ImageIndex);
 				}));
 			};
 
@@ -1035,7 +1032,7 @@ public:
 
 					RUN_TEST(VerifyTextureContents(*TestName, RHICmdList, DestTexture, [&](void* Data, uint32 MipWidth, uint32 MipHeight, uint32 Width, uint32 Height, uint32, uint32 SliceIndex)
 					{
-						return CheckTestData(Data, Width * sizeof(uint32), DestDesc.Extent, FIntPoint::ZeroValue, MipIndex);
+						return CheckTestData(Data, Width * sizeof(uint32), DestDesc.Extent, FIntPoint::ZeroValue, DestDesc.Extent, FIntPoint::ZeroValue, MipIndex);
 					}));
 				}
 			};
@@ -1046,9 +1043,9 @@ public:
 			// Copy region
 			{
 				const TCHAR* TestName = TEXT("Copy Texture Region");
-				const FIntPoint Extent(32, 32);
+				const FIntPoint Extent(64, 64);
 				const FIntPoint PatternOffsetSource(8, 4);
-				const FIntPoint PatternOffsetDest(8, 4);
+				const FIntPoint PatternOffsetDest(16, 20);
 				const FIntPoint PatternSize(16, 24);
 				const int32 ImageIndex = 1;
 
@@ -1062,7 +1059,7 @@ public:
 
 				uint32 Stride = 0;
 				void* Data = RHICmdList.LockTexture2D(SourceTexture, 0, RLM_WriteOnly, Stride, false);
-				WriteTestData(Data, Stride, PatternSize, PatternOffsetSource, ImageIndex);
+				WriteTestData(Data, Stride, Extent, FIntPoint::ZeroValue, ImageIndex);
 				RHICmdList.UnlockTexture2D(SourceTexture, 0, false);
 
 				const FRHITextureCreateDesc DestDesc =
@@ -1082,10 +1079,53 @@ public:
 
 				RUN_TEST(VerifyTextureContents(TestName, RHICmdList, DestTexture, [&](void* Data, uint32 MipWidth, uint32 MipHeight, uint32 Width, uint32 Height, uint32 MipIndex, uint32 SliceIndex)
 				{
-					return CheckTestData(Data, Width * sizeof(uint32), PatternSize, PatternOffsetDest, ImageIndex);
+					return CheckTestData(Data, Width * sizeof(uint32), PatternSize, PatternOffsetDest, Extent, PatternOffsetSource, ImageIndex);
 				}));
 			}
 
+			// Copy source subrect to destination readback buffer
+			{
+				const TCHAR* TestName = TEXT("Readback Texture Region");
+				const FIntPoint Extent(32, 32);
+				const FIntPoint PatternOffset(8, 4);
+				const FIntPoint PatternSize(16, 24);
+				const int32 ImageIndex = 1;
+
+				const FRHITextureCreateDesc SourceDesc =
+					FRHITextureCreateDesc::Create2D(TestName)
+					.SetExtent(Extent.X, Extent.Y)
+					.SetFormat(PF_R32_UINT)
+					.SetInitialState(ERHIAccess::CopySrc);
+
+				FTextureRHIRef SourceTexture = RHICreateTexture(SourceDesc);
+
+				uint32 Stride = 0;
+				void* Data = RHICmdList.LockTexture2D(SourceTexture, 0, RLM_WriteOnly, Stride, false);
+				WriteTestData(Data, Stride, Extent, FIntPoint::ZeroValue, ImageIndex);
+				RHICmdList.UnlockTexture2D(SourceTexture, 0, false);
+
+				FRHIGPUTextureReadback Readback(TEXT("DestReadbackBuffer"));
+				FResolveRect SourceRect = FResolveRect(PatternOffset.X, PatternOffset.Y, PatternOffset.X + PatternSize.X, PatternOffset.Y + PatternSize.Y);
+
+				Readback.EnqueueCopy(RHICmdList, SourceTexture, SourceRect);
+				RHICmdList.SubmitCommandsAndFlushGPU();
+
+				int32 Pitch;
+				void* ReadbackData = Readback.Lock(Pitch);
+				bool bReadbackOK = CheckTestData(ReadbackData, Pitch * sizeof(uint32), PatternSize, FIntPoint::ZeroValue, Extent, PatternOffset, ImageIndex);
+				if (bReadbackOK)
+				{
+					UE_LOG(LogRHIUnitTestCommandlet, Display, TEXT("Test passed. \"%s\""), TestName);
+				}
+				else
+				{
+					UE_LOG(LogRHIUnitTestCommandlet, Error, TEXT("Test failed. \"%s\""), TestName);
+				}
+				Readback.Unlock();
+
+				bResult = bResult && bReadbackOK;
+			}
+
 			// Copy mip chains
 			{
 				const TCHAR* TestName = TEXT("Copy Texture Subresources - Mips");
@@ -1171,7 +1211,7 @@ public:
 					const FIntPoint PatternSize(PatternSizeMip0.X     >> MipIndex, PatternSizeMip0.Y   >> MipIndex);
 					const FIntPoint PatternOffset(PatternOffsetMip0.X >> MipIndex, PatternOffsetMip0.Y >> MipIndex);
 
-					return CheckTestData(Data, Width * sizeof(uint32), PatternSize, PatternOffset, ImageIndex);
+					return CheckTestData(Data, Width * sizeof(uint32), PatternSize, PatternOffset, PatternSize, PatternOffset, ImageIndex);
 				}));
 			}
 
@@ -1246,7 +1286,7 @@ public:
 					const FIntPoint PatternSize(PatternSizeMip0.X >> MipIndex, PatternSizeMip0.Y >> MipIndex);
 					const FIntPoint PatternOffset(PatternOffsetMip0.X >> MipIndex, PatternOffsetMip0.Y >> MipIndex);
 
-					return CheckTestData(Data, Width * sizeof(uint32), PatternSize, PatternOffset, ImageIndex);
+					return CheckTestData(Data, Width * sizeof(uint32), PatternSize, PatternOffset, PatternSize, PatternOffset, ImageIndex);
 				}));
 			}
 
@@ -1279,12 +1319,11 @@ public:
 
 					FUpdateTexture3DData Data = RHICmdList.BeginUpdateTexture3D(SourceTexture, Mip, Region);
 					
-					const FIntPoint PatternOffset(0, 0);
 					const FIntPoint PatternSize(Region.Width, Region.Height);
 
 					for (uint32 Z = 0; Z < Region.Depth; ++Z)
 					{
-						WriteTestData(Data.Data + Data.DepthPitch * Z, Data.RowPitch, PatternSize, PatternOffset, ImageIndex + Z);
+						WriteTestData(Data.Data + Data.DepthPitch * Z, Data.RowPitch, PatternSize, FIntPoint::ZeroValue, ImageIndex + Z);
 					}
 
 					RHICmdList.EndUpdateTexture3D(Data);
@@ -1301,10 +1340,9 @@ public:
 				{
 					int32 ImageIndex = SourceDesc.Depth * MipIndex + Depth + 1;
 
-					const FIntPoint PatternOffset(0, 0);
 					const FIntPoint PatternSize(MipWidth, MipHeight);
 
-					return CheckTestData(Data, Width * sizeof(uint32), PatternSize, PatternOffset, ImageIndex);
+					return CheckTestData(Data, Width * sizeof(uint32), PatternSize, FIntPoint::ZeroValue, PatternSize, FIntPoint::ZeroValue, ImageIndex);
 				}));
 			}
 		}
diff --git a/Engine/Source/Runtime/RHI/Private/RHIGPUReadback.cpp b/Engine/Source/Runtime/RHI/Private/RHIGPUReadback.cpp
--- a/Engine/Source/Runtime/RHI/Private/RHIGPUReadback.cpp
+++ b/Engine/Source/Runtime/RHI/Private/RHIGPUReadback.cpp
@@ -154,28 +154,39 @@ void FRHIGPUTextureReadback::EnqueueCopy(FRHICommandList& RHICmdList, FRHITextur
 			// Assume for now that every enqueue happens on a texture of the same format and size (when reused).
 			if (!DestinationStagingTextures[GPUIndex])
 			{
-				FIntVector TextureSize = SourceTexture->GetSizeXYZ();
+				FIntVector StagingTextureSize;
+				if (Rect.IsValid())
+				{
+					StagingTextureSize = FIntVector(Rect.X2 - Rect.X1, Rect.Y2 - Rect.Y1, 1);
+				}
+				else
+				{
+					StagingTextureSize = SourceTexture->GetSizeXYZ();
+				}
+
 				FString FenceName = Fence->GetFName().ToString();
 
 				const FRHITextureCreateDesc Desc =
-					FRHITextureCreateDesc::Create2D(*FenceName, TextureSize.X, TextureSize.Y, SourceTexture->GetFormat())
+					FRHITextureCreateDesc::Create2D(*FenceName, StagingTextureSize.X, StagingTextureSize.Y, SourceTexture->GetFormat())
 					.SetFlags(ETextureCreateFlags::CPUReadback | ETextureCreateFlags::HideInVisualizeTexture)
 					.SetInitialState(ERHIAccess::CopyDest);
 
 				DestinationStagingTextures[GPUIndex] = RHICreateTexture(Desc);
 			}
-
+			else
+			{
+				RHICmdList.Transition(FRHITransitionInfo(DestinationStagingTextures[GPUIndex], ERHIAccess::CPURead, ERHIAccess::CopyDest));
+			}
 			FRHICopyTextureInfo CopyInfo;
 
 			if (Rect.IsValid())
 			{
 				CopyInfo.SourcePosition = FIntVector(Rect.X1, Rect.Y1, 0);
-				CopyInfo.DestPosition = CopyInfo.SourcePosition;
 				CopyInfo.Size = FIntVector(Rect.X2 - Rect.X1, Rect.Y2 - Rect.Y1, 1);
 			}
 
 			RHICmdList.CopyTexture(SourceTexture, DestinationStagingTextures[GPUIndex], CopyInfo);
-
+			RHICmdList.Transition(FRHITransitionInfo(DestinationStagingTextures[GPUIndex], ERHIAccess::CopyDest, ERHIAccess::CPURead));
 			RHICmdList.WriteGPUFence(Fence);
 		}
 	}
@@ -192,7 +203,6 @@ void* FRHIGPUTextureReadback::Lock(uint32 NumBytes)
 		void* ResultsBuffer = nullptr;
 		int32 BufferWidth = 0, BufferHeight = 0;
 		FRHICommandListImmediate& RHICmdList = FRHICommandListExecutor::GetImmediateCommandList();
-		RHICmdList.Transition(FRHITransitionInfo(DestinationStagingTextures[GPUIndex], ERHIAccess::CopyDest, ERHIAccess::CPURead));
 		RHICmdList.MapStagingSurface(DestinationStagingTextures[GPUIndex], Fence.GetReference(), ResultsBuffer, BufferWidth, BufferHeight, LastLockGPUIndex);
 
 		return ResultsBuffer;
@@ -214,7 +224,6 @@ void* FRHIGPUTextureReadback::Lock(int32& OutRowPitchInPixels, int32 *OutBufferH
 		void* ResultsBuffer = nullptr;
 		int32 BufferWidth = 0, BufferHeight = 0;
 		FRHICommandListImmediate& RHICmdList = FRHICommandListExecutor::GetImmediateCommandList();
-		RHICmdList.Transition(FRHITransitionInfo(DestinationStagingTextures[GPUIndex], ERHIAccess::CopyDest, ERHIAccess::CPURead));
 		RHICmdList.MapStagingSurface(DestinationStagingTextures[GPUIndex], Fence.GetReference(), ResultsBuffer, BufferWidth, BufferHeight, LastLockGPUIndex);
 
 		if (OutBufferHeight)
@@ -246,7 +255,6 @@ void FRHIGPUTextureReadback::Unlock()
 
 	FRHICommandListImmediate& RHICmdList = FRHICommandListExecutor::GetImmediateCommandList();
 	RHICmdList.UnmapStagingSurface(DestinationStagingTextures[LastLockGPUIndex], LastLockGPUIndex);
-	RHICmdList.Transition(FRHITransitionInfo(DestinationStagingTextures[LastLockGPUIndex], ERHIAccess::CPURead, ERHIAccess::CopyDest));
 }
 
 uint64 FRHIGPUTextureReadback::GetGPUSizeBytes() const
