Add RHI unit test for multi draw indirect

diff --git a/Engine/Plugins/Tests/RHITests/Shaders/Private/TestDrawInstanced.usf b/Engine/Plugins/Tests/RHITests/Shaders/Private/TestDrawInstanced.usf
new file mode 100644
--- /dev/null
+++ b/Engine/Plugins/Tests/RHITests/Shaders/Private/TestDrawInstanced.usf
@@ -0,0 +1,41 @@
+// Copyright Epic Games, Inc. All Rights Reserved.
+
+#include "/Engine/Private/Common.ush"
+
+// Basic shader that writes 1 to the output buffer entries based on which instances were drawn.
+// Used to validate that instanced draws with base instance offset work correctly.
+
+void TestDrawInstancedMainVS(
+	in uint CustomInstanceID  : ATTRIBUTE0,
+	in uint VertexID		: SV_VertexID,
+	out uint OutInstanceID  : Texcoord0,
+	out float4 OutPosition	: SV_Position
+)
+{
+	OutInstanceID = CustomInstanceID;
+
+#if 0
+	// Produces broken code on some platforms
+	switch(VertexID)
+	{
+		default:
+		case 0:	OutPosition = float4(-1, -1, 0, 1); break;
+		case 1:	OutPosition = float4(-1, +3, 0, 1); break;
+		case 2:	OutPosition = float4(+3, -1, 0, 1); break;
+	}
+#else
+	if      (VertexID == 0) { OutPosition = float4(-1, -1, 0, 1); }
+	else if (VertexID == 1) { OutPosition = float4(-1, +3, 0, 1); }
+	else                    { OutPosition = float4(+3, -1, 0, 1); }
+#endif
+}
+
+RWBuffer<uint>	OutDrawnInstances;
+void TestDrawInstancedMainPS(
+	in nointerpolation uint InstanceID  : Texcoord0,
+	out float4 OutColor	: SV_Target0
+)
+{
+	OutDrawnInstances[InstanceID] = 1;
+	OutColor = float4(1,1,1,1);
+}
diff --git a/Engine/Plugins/Tests/RHITests/Source/RHITests/Private/RHIDrawTests.cpp b/Engine/Plugins/Tests/RHITests/Source/RHITests/Private/RHIDrawTests.cpp
new file mode 100644
--- /dev/null
+++ b/Engine/Plugins/Tests/RHITests/Source/RHITests/Private/RHIDrawTests.cpp
@@ -0,0 +1,181 @@
+// Copyright Epic Games, Inc. All Rights Reserved.
+
+#include "RHIDrawTests.h"
+#include "RHIBufferTests.h" // for VerifyBufferContents
+#include "CommonRenderResources.h"
+#include "RenderCaptureInterface.h"
+
+class FTestDrawInstancedVS : public FGlobalShader
+{
+	DECLARE_GLOBAL_SHADER(FTestDrawInstancedVS);
+	FTestDrawInstancedVS() = default;
+	FTestDrawInstancedVS(const ShaderMetaType::CompiledShaderInitializerType& Initializer)
+		: FGlobalShader(Initializer)
+	{
+	}
+	static void ModifyCompilationEnvironment(const FGlobalShaderPermutationParameters& Parameters, FShaderCompilerEnvironment& OutEnvironment)
+	{
+	}
+};
+IMPLEMENT_GLOBAL_SHADER(FTestDrawInstancedVS, "/Plugin/RHITests/Private/TestDrawInstanced.usf", "TestDrawInstancedMainVS", SF_Vertex);
+
+class FTestDrawInstancedPS : public FGlobalShader
+{
+	DECLARE_GLOBAL_SHADER(FTestDrawInstancedPS);
+	FTestDrawInstancedPS() = default;
+	FTestDrawInstancedPS(const ShaderMetaType::CompiledShaderInitializerType& Initializer)
+		: FGlobalShader(Initializer)
+	{
+		OutDrawnInstances.Bind(Initializer.ParameterMap, TEXT("OutDrawnInstances"), SPF_Mandatory);
+	}
+	static void ModifyCompilationEnvironment(const FGlobalShaderPermutationParameters& Parameters, FShaderCompilerEnvironment& OutEnvironment)
+	{
+	}
+	LAYOUT_FIELD(FShaderParameter, OutDrawnInstances);
+};
+IMPLEMENT_GLOBAL_SHADER(FTestDrawInstancedPS, "/Plugin/RHITests/Private/TestDrawInstanced.usf", "TestDrawInstancedMainPS", SF_Pixel);
+
+template <typename T>
+static FBufferRHIRef CreateBufferWithData(EBufferUsageFlags UsageFlags, ERHIAccess ResourceState, const TCHAR* Name, TConstArrayView<T> Data)
+{
+	uint32 BufferSize = sizeof(T) * Data.Num();
+	FRHIResourceCreateInfo CreateInfo(Name);
+	FBufferRHIRef Buffer = RHICreateBuffer(BufferSize, UsageFlags, sizeof(T), ResourceState, CreateInfo);
+	void* MappedData = RHILockBuffer(Buffer, 0, BufferSize, EResourceLockMode::RLM_WriteOnly);
+	FMemory::Memcpy(MappedData, Data.GetData(), BufferSize);
+	RHIUnlockBuffer(Buffer);
+	return Buffer;
+}
+
+bool FRHIDrawTests::Test_MultiDrawIndirect(FRHICommandListImmediate& RHICmdList)
+{
+	if (!GRHIGlobals.SupportsMultiDrawIndirect)
+	{
+		return true;
+	}
+
+	// Probably could/should automatically enable in the outer scope when running RHI Unit Tests
+	// RenderCaptureInterface::FScopedCapture RenderCapture(true /*bEnable*/, &RHICmdList, TEXT("Test_MultiDrawIndirect"));
+
+	static constexpr uint32 MaxInstances = 8;
+
+	// D3D12 does not have a way to get the base instance ID (SV_InstanceID always starts from 0), so we must emulate it...
+	const uint32 InstanceIDs[MaxInstances] = { 0, 1, 2, 3, 4, 5, 6, 7 };
+	FBufferRHIRef InstanceIDBuffer = CreateBufferWithData(EBufferUsageFlags::VertexBuffer, ERHIAccess::VertexOrIndexBuffer, TEXT("Test_MultiDrawIndirect_InstanceID"), MakeArrayView(InstanceIDs));
+
+	FVertexDeclarationElementList VertexDeclarationElements;
+	VertexDeclarationElements.Add(FVertexElement(0, 0, VET_UInt, 0, 4, true /*per instance frequency*/));
+	FVertexDeclarationRHIRef VertexDeclarationRHI = PipelineStateCache::GetOrCreateVertexDeclaration(VertexDeclarationElements);
+
+	const uint16 Indices[3] = { 0, 1, 2 };
+	FBufferRHIRef IndexBuffer = CreateBufferWithData(EBufferUsageFlags::IndexBuffer, ERHIAccess::VertexOrIndexBuffer, TEXT("Test_MultiDrawIndirect_IndexBuffer"), MakeArrayView(Indices));
+
+	static constexpr uint32 OutputBufferStride = sizeof(uint32);
+	static constexpr uint32 OutputBufferSize = OutputBufferStride * MaxInstances;
+	FRHIResourceCreateInfo OutputBufferCreateInfo(TEXT("Test_MultiDrawIndirect_IndexBuffer"));
+	FBufferRHIRef OutputBuffer = RHICreateBuffer(OutputBufferSize, EBufferUsageFlags::UnorderedAccess, OutputBufferStride, ERHIAccess::UAVGraphics, OutputBufferCreateInfo);
+	{
+		void* BufferData = RHILockBuffer(OutputBuffer, 0, OutputBufferSize, EResourceLockMode::RLM_WriteOnly);
+		FMemory::Memzero(BufferData, OutputBufferSize);
+		RHIUnlockBuffer(OutputBuffer);
+	}
+
+	const uint32 CountValues[4] = { 1, 1, 16, 0 };
+	FBufferRHIRef CountBuffer = CreateBufferWithData(EBufferUsageFlags::DrawIndirect | EBufferUsageFlags::UnorderedAccess, ERHIAccess::IndirectArgs, TEXT("Test_MultiDrawIndirect_Count"), MakeArrayView(CountValues));
+
+	const FRHIDrawIndexedIndirectParameters DrawArgs[] =
+	{
+		// IndexCountPerInstance, InstanceCount, StartIndexLocation, BaseVertexLocation, StartInstanceLocation
+		{3, 1, 0, 0, 0}, // fill slot 0
+		// gap in slot 1
+		{3, 2, 0, 0, 2}, // fill slots 2, 3 using 1 sub-draw
+		// gap in slot 4
+		{3, 1, 0, 0, 5}, // fill slots 5, 6 using 2 sub-draws
+		{3, 1, 0, 0, 6},
+		{3, 1, 0, 0, 7}, // this draw is expected to never execute
+	};
+
+	const uint32 ExpectedDrawnInstances[MaxInstances] = { 1, 0, 1, 1, 0, 1, 1, 0 };
+
+	FBufferRHIRef DrawArgBuffer = CreateBufferWithData(EBufferUsageFlags::DrawIndirect | EBufferUsageFlags::UnorderedAccess, ERHIAccess::IndirectArgs,
+		TEXT("Test_MultiDrawIndirect_DrawArgs"), MakeArrayView(DrawArgs));
+
+	FUnorderedAccessViewRHIRef OutputBufferUAV = RHICreateUnorderedAccessView(OutputBuffer, 
+		FRHIViewDesc::CreateBufferUAV()
+		.SetType(FRHIViewDesc::EBufferType::Typed)
+		.SetFormat(PF_R32_UINT));
+
+	const FIntPoint RenderTargetSize(4, 4);
+	FRHITextureDesc RenderTargetTextureDesc(ETextureDimension::Texture2D, ETextureCreateFlags::RenderTargetable, PF_B8G8R8A8, FClearValueBinding(), RenderTargetSize, 1, 1, 1, 1, 0);
+	FRHITextureCreateDesc RenderTargetCreateDesc(RenderTargetTextureDesc, ERHIAccess::RTV, TEXT("Test_MultiDrawIndirect_RenderTarget"));
+	FTextureRHIRef RenderTarget = RHICreateTexture(RenderTargetCreateDesc);
+
+	TShaderMapRef<FTestDrawInstancedVS> VertexShader(GetGlobalShaderMap(GMaxRHIFeatureLevel));
+	TShaderMapRef<FTestDrawInstancedPS> PixelShader(GetGlobalShaderMap(GMaxRHIFeatureLevel));
+
+	FGraphicsPipelineStateInitializer GraphicsPSOInit;
+	
+	GraphicsPSOInit.BoundShaderState.VertexShaderRHI = VertexShader.GetVertexShader();
+	GraphicsPSOInit.BoundShaderState.VertexDeclarationRHI = VertexDeclarationRHI;
+	GraphicsPSOInit.BoundShaderState.PixelShaderRHI = PixelShader.GetPixelShader();
+	GraphicsPSOInit.DepthStencilState = TStaticDepthStencilState<false, CF_Always>::GetRHI();
+	GraphicsPSOInit.BlendState = TStaticBlendState<>::GetRHI();
+	GraphicsPSOInit.RasterizerState = TStaticRasterizerState<>::GetRHI();
+	GraphicsPSOInit.PrimitiveType = EPrimitiveType::PT_TriangleList;
+
+	FRHITexture* ColorRTs[1] = { RenderTarget.GetReference() };
+	FRHIRenderPassInfo RenderPassInfo(1, ColorRTs, ERenderTargetActions::DontLoad_DontStore);
+
+	RHICmdList.BeginRenderPass(RenderPassInfo, TEXT("Test_MultiDrawIndirect"));
+	RHICmdList.SetViewport(0, 0, 0, float(RenderTargetSize.X), float(RenderTargetSize.Y), 1);
+
+	RHICmdList.ApplyCachedRenderTargets(GraphicsPSOInit);
+	SetGraphicsPipelineState(RHICmdList, GraphicsPSOInit, 0);
+
+	check(InstanceIDBuffer->GetStride() == 4);
+	RHICmdList.SetStreamSource(0, InstanceIDBuffer, 0);
+
+	FRHIBatchedShaderParameters ShaderParameters;
+	ShaderParameters.SetUAVParameter(PixelShader->OutDrawnInstances.GetBaseIndex(), OutputBufferUAV);
+	RHICmdList.SetBatchedShaderParameters(PixelShader.GetPixelShader(), ShaderParameters);
+
+	const uint32 DrawArgsStride = sizeof(DrawArgs[0]);
+	const uint32 CountStride = sizeof(CountValues[0]);
+
+	RHICmdList.MultiDrawIndexedPrimitiveIndirect(IndexBuffer, 
+		DrawArgBuffer, DrawArgsStride*0, // 1 sub-draw with instance index 0
+		CountBuffer, CountStride*0, // count buffer contains 1 in this slot
+		5 // expect to draw only 1 instance due to GPU-side upper bound
+	);
+
+	RHICmdList.MultiDrawIndexedPrimitiveIndirect(IndexBuffer,
+		DrawArgBuffer, DrawArgsStride*1, // 1 sub-draw with 2 instances at base index 2
+		CountBuffer, CountStride*1, // count buffer contains 1 in this slot
+		4 // expect to draw only 1 instance due to GPU-side upper bound
+	);
+
+	RHICmdList.MultiDrawIndexedPrimitiveIndirect(IndexBuffer,
+		DrawArgBuffer, DrawArgsStride*2, // 2 sub-draws with 1 instance each starting at base index 5
+		CountBuffer, CountStride*2, // count buffer contains 16 in this slot
+		2 // expect to draw only 2 instances due to CPU-side upper bound
+	);
+
+	RHICmdList.MultiDrawIndexedPrimitiveIndirect(IndexBuffer,
+		DrawArgBuffer, DrawArgsStride*4, // 1 sub-draw with 1 instance each starting at base index 7
+		CountBuffer, CountStride*3, // count buffer contains 0 in this slot
+		1 // expect to skip the draw due to GPU-side count of 0
+	);
+
+	RHICmdList.MultiDrawIndexedPrimitiveIndirect(IndexBuffer,
+		DrawArgBuffer, DrawArgsStride*4, // 1 sub-draw with 1 instance each starting at base index 7
+		CountBuffer, CountStride*0, // count buffer contains 1 in this slot
+		0 // expect to skip the draw due to CPU-side count of 0
+	);
+
+	RHICmdList.EndRenderPass();
+
+	TConstArrayView<uint8> ExpectedOutputView = MakeArrayView(reinterpret_cast<const uint8*>(ExpectedDrawnInstances), sizeof(ExpectedDrawnInstances));
+	bool bSucceeded = FRHIBufferTests::VerifyBufferContents(TEXT("Test_MultiDrawIndirect"), RHICmdList, OutputBuffer, ExpectedOutputView);
+
+	return bSucceeded;
+}
diff --git a/Engine/Plugins/Tests/RHITests/Source/RHITests/Private/RHITests.spec.cpp b/Engine/Plugins/Tests/RHITests/Source/RHITests/Private/RHITests.spec.cpp
--- a/Engine/Plugins/Tests/RHITests/Source/RHITests/Private/RHITests.spec.cpp
+++ b/Engine/Plugins/Tests/RHITests/Source/RHITests/Private/RHITests.spec.cpp
@@ -6,7 +6,7 @@
 #include "RenderingThread.h"
 #include "RHIBufferTests.h"
 #include "RHITextureTests.h"
-
+#include "RHIDrawTests.h"
 
 BEGIN_DEFINE_SPEC(FAutomationRHITest, "Rendering.RHI", EAutomationTestFlags::EngineFilter | EAutomationTestFlags::ApplicationContextMask | EAutomationTestFlags::NonNullRHI)
 END_DEFINE_SPEC(FAutomationRHITest)
@@ -92,4 +92,13 @@ void FAutomationRHITest::Define()
 			TestEqual("RHICreateBuffer_Parallel", bResult, 1);
 		});
 	});
+
+	Describe("Test RHI Draw", [this]()
+	{
+		It("RHI MultiDrawIndirect", [this]()
+		{
+			bool bResult = RunOnRenderThreadSynchronous(FRHIDrawTests::Test_MultiDrawIndirect);
+			TestEqual("RHI MultiDrawIndirect", bResult, 1);
+		});
+	});
 }
diff --git a/Engine/Plugins/Tests/RHITests/Source/RHITests/Private/RHITestsModule.cpp b/Engine/Plugins/Tests/RHITests/Source/RHITests/Private/RHITestsModule.cpp
--- a/Engine/Plugins/Tests/RHITests/Source/RHITests/Private/RHITestsModule.cpp
+++ b/Engine/Plugins/Tests/RHITests/Source/RHITests/Private/RHITestsModule.cpp
@@ -7,6 +7,7 @@
 #include "ShaderCore.h"
 #include "RHIBufferTests.h"
 #include "RHITextureTests.h"
+#include "RHIDrawTests.h"
 
 #define LOCTEXT_NAMESPACE "FRHITestsModule"
 
@@ -14,6 +15,13 @@ static bool RunTests_RenderThread(FRHICommandListImmediate& RHICmdList)
 {
 	bool bResult = true;
 
+	// ------------------------------------------------
+	// Drawing
+	// ------------------------------------------------
+	{
+		RUN_TEST(FRHIDrawTests::Test_MultiDrawIndirect(RHICmdList));
+	}
+
 	// ------------------------------------------------
 	// RHI Formats
 	// ------------------------------------------------
@@ -59,8 +67,8 @@ static bool RunTests_RenderThread(FRHICommandListImmediate& RHICmdList)
 
 void FRHITestsModule::StartupModule()
 {
-	//FString PluginShaderDir = FPaths::Combine(IPluginManager::Get().FindPlugin(TEXT("RHITests"))->GetBaseDir(), TEXT("Shaders"));
-	//AddShaderSourceDirectoryMapping(TEXT("/Plugin/RHITests"), PluginShaderDir);
+	FString PluginShaderDir = FPaths::Combine(IPluginManager::Get().FindPlugin(TEXT("RHITests"))->GetBaseDir(), TEXT("Shaders"));
+	AddShaderSourceDirectoryMapping(TEXT("/Plugin/RHITests"), PluginShaderDir);
 }
 
 void FRHITestsModule::RunAllTests()
diff --git a/Engine/Plugins/Tests/RHITests/Source/RHITests/Public/RHIBufferTests.h b/Engine/Plugins/Tests/RHITests/Source/RHITests/Public/RHIBufferTests.h
--- a/Engine/Plugins/Tests/RHITests/Source/RHITests/Public/RHIBufferTests.h
+++ b/Engine/Plugins/Tests/RHITests/Source/RHITests/Public/RHIBufferTests.h
@@ -7,10 +7,14 @@
 
 class FRHIBufferTests
 {
+public:
+
 	// Copies data in the specified vertex buffer back to the CPU, and passes a pointer to that data to the provided verification lambda.
 	static bool VerifyBufferContents(const TCHAR* TestName, FRHICommandListImmediate& RHICmdList, TArrayView<FRHIBuffer*> Buffers, TFunctionRef<bool(int32 BufferIndex, void* Ptr, uint32 NumBytes)> VerifyCallback);
 	static bool VerifyBufferContents(const TCHAR* TestName, FRHICommandListImmediate& RHICmdList, FRHIBuffer* Buffer, TConstArrayView<uint8> ExpectedData);
 
+private:
+
 	template <typename TestLambdaType>
 	static void ParallelDispatchCommands(FRHICommandListImmediate& RHICmdList, int32 NumTests, TestLambdaType TestLambda)
 	{
diff --git a/Engine/Plugins/Tests/RHITests/Source/RHITests/Public/RHIDrawTests.h b/Engine/Plugins/Tests/RHITests/Source/RHITests/Public/RHIDrawTests.h
new file mode 100644
--- /dev/null
+++ b/Engine/Plugins/Tests/RHITests/Source/RHITests/Public/RHIDrawTests.h
@@ -0,0 +1,11 @@
+// Copyright Epic Games, Inc. All Rights Reserved.
+
+#include "RHITestsCommon.h"
+
+class FRHIDrawTests
+{
+public:
+
+	static bool Test_MultiDrawIndirect(FRHICommandListImmediate& RHICmdList);
+};
+
