Fix missing pragma once

diff --git a/Engine/Plugins/Tests/RHITests/Source/RHITests/Public/RHIDrawTests.h b/Engine/Plugins/Tests/RHITests/Source/RHITests/Public/RHIDrawTests.h
--- a/Engine/Plugins/Tests/RHITests/Source/RHITests/Public/RHIDrawTests.h
+++ b/Engine/Plugins/Tests/RHITests/Source/RHITests/Public/RHIDrawTests.h
@@ -1,5 +1,7 @@
 // Copyright Epic Games, Inc. All Rights Reserved.
 
+#pragma once
+
 #include "RHITestsCommon.h"
 
 class FRHIDrawTests
