Fix non-unity compile error due to missing RHIStaticStates.h include

diff --git a/Engine/Plugins/Tests/RHITests/Source/RHITests/Private/RHIDrawTests.cpp b/Engine/Plugins/Tests/RHITests/Source/RHITests/Private/RHIDrawTests.cpp
--- a/Engine/Plugins/Tests/RHITests/Source/RHITests/Private/RHIDrawTests.cpp
+++ b/Engine/Plugins/Tests/RHITests/Source/RHITests/Private/RHIDrawTests.cpp
@@ -4,6 +4,7 @@
 #include "RHIBufferTests.h" // for VerifyBufferContents
 #include "CommonRenderResources.h"
 #include "RenderCaptureInterface.h"
+#include "RHIStaticStates.h"
 
 class FTestDrawInstancedVS : public FGlobalShader
 {
