Fix non-determinism in random number genration in Test_RHICreateBuffer_Parallel between platforms

diff --git a/Engine/Plugins/Tests/RHITests/Source/RHITests/Public/RHIBufferTests.h b/Engine/Plugins/Tests/RHITests/Source/RHITests/Public/RHIBufferTests.h
--- a/Engine/Plugins/Tests/RHITests/Source/RHITests/Public/RHIBufferTests.h
+++ b/Engine/Plugins/Tests/RHITests/Source/RHITests/Public/RHIBufferTests.h
@@ -4,6 +4,7 @@
 
 #include "RHITestsCommon.h"
 #include "Async/ParallelFor.h"
+#include "Math/RandomStream.h"
 
 class FRHIBufferTests
 {
@@ -411,7 +412,7 @@ public:
 
 		const int32 NumBuffersToCreate = 256;
 
-		FGenericPlatformMath::RandInit(1);
+		FRandomStream Rng(0x5bf8a575);
 
 		TArray<TRefCountPtr<FRHIBuffer>> Buffers;
 		Buffers.SetNum(NumBuffersToCreate);
@@ -421,7 +422,7 @@ public:
 
 		for (int32 Index = 0; Index < NumBuffersToCreate; ++Index)
 		{
-			RandomNumberPerBuffer[Index] = FGenericPlatformMath::Rand();
+			RandomNumberPerBuffer[Index] = Rng.RandRange(1, 65536);
 		}
 
 		ParallelDispatchCommands(RHICmdList, NumBuffersToCreate, [&](FRHICommandList& InRHICmdList, int32 Index, int32 Num)
@@ -429,7 +430,7 @@ public:
 			SCOPED_NAMED_EVENT_TEXT("TestCreateBuffer", FColor::Magenta);
 
 			const int32 Random = RandomNumberPerBuffer[Index];
-			const int32 BufferSize = Align(Random % 65536, 16);
+			const int32 BufferSize = Align(Random, 16);
 			const int32 BufferStride = 4;
 
 			EBufferUsageFlags Usage = EBufferUsageFlags::VertexBuffer | EBufferUsageFlags::SourceCopy;
