Fix RHI validation errors in Test_MultiDrawIndirect

diff --git a/Engine/Plugins/Tests/RHITests/Source/RHITests/Private/RHIDrawTests.cpp b/Engine/Plugins/Tests/RHITests/Source/RHITests/Private/RHIDrawTests.cpp
--- a/Engine/Plugins/Tests/RHITests/Source/RHITests/Private/RHIDrawTests.cpp
+++ b/Engine/Plugins/Tests/RHITests/Source/RHITests/Private/RHIDrawTests.cpp
@@ -73,13 +73,8 @@ bool FRHIDrawTests::Test_MultiDrawIndirect(FRHICommandListImmediate& RHICmdList)
 
 	static constexpr uint32 OutputBufferStride = sizeof(uint32);
 	static constexpr uint32 OutputBufferSize = OutputBufferStride * MaxInstances;
-	FRHIResourceCreateInfo OutputBufferCreateInfo(TEXT("Test_MultiDrawIndirect_IndexBuffer"));
-	FBufferRHIRef OutputBuffer = RHICreateBuffer(OutputBufferSize, EBufferUsageFlags::UnorderedAccess, OutputBufferStride, ERHIAccess::UAVGraphics, OutputBufferCreateInfo);
-	{
-		void* BufferData = RHILockBuffer(OutputBuffer, 0, OutputBufferSize, EResourceLockMode::RLM_WriteOnly);
-		FMemory::Memzero(BufferData, OutputBufferSize);
-		RHIUnlockBuffer(OutputBuffer);
-	}
+	FRHIResourceCreateInfo OutputBufferCreateInfo(TEXT("Test_MultiDrawIndirect_OutputBuffer"));
+	FBufferRHIRef OutputBuffer = RHICreateBuffer(OutputBufferSize, EBufferUsageFlags::UnorderedAccess | EBufferUsageFlags::SourceCopy, OutputBufferStride, ERHIAccess::UAVCompute, OutputBufferCreateInfo);
 
 	const uint32 CountValues[4] = { 1, 1, 16, 0 };
 	FBufferRHIRef CountBuffer = CreateBufferWithData(EBufferUsageFlags::DrawIndirect | EBufferUsageFlags::UnorderedAccess, ERHIAccess::IndirectArgs, TEXT("Test_MultiDrawIndirect_Count"), MakeArrayView(CountValues));
@@ -98,7 +93,7 @@ bool FRHIDrawTests::Test_MultiDrawIndirect(FRHICommandListImmediate& RHICmdList)
 
 	const uint32 ExpectedDrawnInstances[MaxInstances] = { 1, 0, 1, 1, 0, 1, 1, 0 };
 
-	FBufferRHIRef DrawArgBuffer = CreateBufferWithData(EBufferUsageFlags::DrawIndirect | EBufferUsageFlags::UnorderedAccess, ERHIAccess::IndirectArgs,
+	FBufferRHIRef DrawArgBuffer = CreateBufferWithData(EBufferUsageFlags::DrawIndirect | EBufferUsageFlags::UnorderedAccess | EBufferUsageFlags::VertexBuffer, ERHIAccess::IndirectArgs,
 		TEXT("Test_MultiDrawIndirect_DrawArgs"), MakeArrayView(DrawArgs));
 
 	FUnorderedAccessViewRHIRef OutputBufferUAV = RHICreateUnorderedAccessView(OutputBuffer, 
@@ -106,6 +101,8 @@ bool FRHIDrawTests::Test_MultiDrawIndirect(FRHICommandListImmediate& RHICmdList)
 		.SetType(FRHIViewDesc::EBufferType::Typed)
 		.SetFormat(PF_R32_UINT));
 
+	RHICmdList.ClearUAVUint(OutputBufferUAV, FUintVector4(0));
+
 	const FIntPoint RenderTargetSize(4, 4);
 	FRHITextureDesc RenderTargetTextureDesc(ETextureDimension::Texture2D, ETextureCreateFlags::RenderTargetable, PF_B8G8R8A8, FClearValueBinding(), RenderTargetSize, 1, 1, 1, 1, 0);
 	FRHITextureCreateDesc RenderTargetCreateDesc(RenderTargetTextureDesc, ERHIAccess::RTV, TEXT("Test_MultiDrawIndirect_RenderTarget"));
@@ -127,6 +124,9 @@ bool FRHIDrawTests::Test_MultiDrawIndirect(FRHICommandListImmediate& RHICmdList)
 	FRHITexture* ColorRTs[1] = { RenderTarget.GetReference() };
 	FRHIRenderPassInfo RenderPassInfo(1, ColorRTs, ERenderTargetActions::DontLoad_DontStore);
 
+	RHICmdList.Transition(FRHITransitionInfo(OutputBufferUAV, ERHIAccess::UAVCompute, ERHIAccess::UAVGraphics, EResourceTransitionFlags::None));
+	RHICmdList.BeginUAVOverlap(); // Output UAV can be written without syncs between draws (each draw is expected to write into different slots)
+
 	RHICmdList.BeginRenderPass(RenderPassInfo, TEXT("Test_MultiDrawIndirect"));
 	RHICmdList.SetViewport(0, 0, 0, float(RenderTargetSize.X), float(RenderTargetSize.Y), 1);
 
@@ -175,6 +175,10 @@ bool FRHIDrawTests::Test_MultiDrawIndirect(FRHICommandListImmediate& RHICmdList)
 
 	RHICmdList.EndRenderPass();
 
+	RHICmdList.EndUAVOverlap();
+
+	RHICmdList.Transition(FRHITransitionInfo(OutputBufferUAV, ERHIAccess::UAVGraphics, ERHIAccess::CopySrc, EResourceTransitionFlags::None));
+
 	TConstArrayView<uint8> ExpectedOutputView = MakeArrayView(reinterpret_cast<const uint8*>(ExpectedDrawnInstances), sizeof(ExpectedDrawnInstances));
 	bool bSucceeded = FRHIBufferTests::VerifyBufferContents(TEXT("Test_MultiDrawIndirect"), RHICmdList, OutputBuffer, ExpectedOutputView);
 
