Add support for volumes and texture arrays in FRHIGPUTextureReadback.

#rnx

diff --git a/Engine/Plugins/Tests/RHITests/Shaders/Private/TestFillBuffer.usf b/Engine/Plugins/Tests/RHITests/Shaders/Private/TestFillBuffer.usf
new file mode 100644
--- /dev/null
+++ b/Engine/Plugins/Tests/RHITests/Shaders/Private/TestFillBuffer.usf
@@ -0,0 +1,59 @@
+// Copyright Epic Games, Inc. All Rights Reserved.
+
+#include "/Engine/Private/Common.ush"
+
+struct FStructuredBufferElem
+{
+	uint	IntVal;
+	float	FloatVal;
+};
+
+RWBuffer<uint>								TypedBufferUAV;
+RWByteAddressBuffer							ByteAddressBufferUAV;
+RWStructuredBuffer<FStructuredBufferElem>	StructuredBufferUAV;
+uint										NumElements;
+
+[numthreads(THREADGROUP_SIZE_X, THREADGROUP_SIZE_Y, 1)]
+void FillBuffers(in uint3 ThreadID : SV_DispatchThreadID)
+{
+	uint Index = ThreadID.x;
+	if(Index < NumElements)
+	{
+		uint TypedVal = Index;
+		uint ByteAddressVal = Index + 1024;
+		uint StructuredIntVal = Index + 2048;
+		uint StructuredFloatVal = (float)(Index + 4096);
+
+		TypedBufferUAV[Index] = TypedVal;
+		ByteAddressBufferUAV.Store(Index * 4, ByteAddressVal);
+		StructuredBufferUAV[Index].IntVal = StructuredIntVal;
+		StructuredBufferUAV[Index].FloatVal = StructuredFloatVal;
+	}
+}
+
+RWTexture2D<uint> Texture2DUAV;
+RWTexture3D<uint> Texture3DUAV;
+RWTexture2DArray<uint> TextureArrayUAV;
+uint3 TextureSize;
+
+[numthreads(THREADGROUP_SIZE_X, THREADGROUP_SIZE_Y, 1)]
+void FillTextures(in uint3 ThreadID : SV_DispatchThreadID)
+{
+	if(ThreadID.x < TextureSize.x && ThreadID.y < TextureSize.y)
+	{
+		uint Texture2DVal = ThreadID.y * TextureSize.x + ThreadID.x;
+		Texture2DUAV[ThreadID.xy] = Texture2DVal;
+
+		for(int Z = 0; Z < TextureSize.z; ++Z)
+		{
+			uint Texture3DVal = Texture2DVal + (Z + 10) * 1024;
+			Texture3DUAV[uint3(ThreadID.xy, Z)] = Texture3DVal;
+		}
+
+		for(int Slice = 0; Slice < TextureSize.z; ++Slice)
+		{
+			uint TextureArrayVal = Texture2DVal + (Slice + 10 + TextureSize.z) * 1024;
+			TextureArrayUAV[uint3(ThreadID.xy, Slice)] = TextureArrayVal;
+		}
+	}
+}
diff --git a/Engine/Plugins/Tests/RHITests/Source/RHITests/Private/RHIReadbackTests.cpp b/Engine/Plugins/Tests/RHITests/Source/RHITests/Private/RHIReadbackTests.cpp
new file mode 100644
--- /dev/null
+++ b/Engine/Plugins/Tests/RHITests/Source/RHITests/Private/RHIReadbackTests.cpp
@@ -0,0 +1,426 @@
+// Copyright Epic Games, Inc. All Rights Reserved.
+
+#include "RHIReadbackTests.h"
+#include "CommonRenderResources.h"
+#include "RHIFeatureLevel.h"
+#include "ShaderParameterMacros.h"
+#include "RenderGraphUtils.h"
+#include "RHIGPUReadback.h"
+
+struct FStructuredBufferElem
+{
+	uint32	IntVal;
+	float	FloatVal;
+};
+
+class FFillBufferCS : public FGlobalShader
+{
+	DECLARE_GLOBAL_SHADER(FFillBufferCS);
+	SHADER_USE_PARAMETER_STRUCT(FFillBufferCS, FGlobalShader);
+
+	static constexpr uint32 ThreadGroupSizeX = 64;
+	static constexpr uint32 ThreadGroupSizeY = 1;
+
+	static void ModifyCompilationEnvironment(const FGlobalShaderPermutationParameters& Parameters, FShaderCompilerEnvironment& OutEnvironment)
+	{
+		FGlobalShader::ModifyCompilationEnvironment(Parameters, OutEnvironment);
+		OutEnvironment.SetDefine(TEXT("THREADGROUP_SIZE_X"), ThreadGroupSizeX);
+		OutEnvironment.SetDefine(TEXT("THREADGROUP_SIZE_Y"), ThreadGroupSizeY);
+	}
+
+	BEGIN_SHADER_PARAMETER_STRUCT(FParameters, )
+		SHADER_PARAMETER_UAV(RWBuffer<uint>, TypedBufferUAV)
+		SHADER_PARAMETER_UAV(RWByteAddressBuffer, ByteAddressBufferUAV)
+		SHADER_PARAMETER_UAV(RWStructuredBuffer<FStructuredBufferElem>, StructuredBufferUAV)
+		SHADER_PARAMETER(uint32, NumElements)
+	END_SHADER_PARAMETER_STRUCT()
+};
+
+IMPLEMENT_GLOBAL_SHADER(FFillBufferCS, "/Plugin/RHITests/Private/TestFillBuffer.usf", "FillBuffers", SF_Compute);
+
+class FFillTexturesCS : public FGlobalShader
+{
+	DECLARE_GLOBAL_SHADER(FFillTexturesCS);
+	SHADER_USE_PARAMETER_STRUCT(FFillTexturesCS, FGlobalShader);
+
+	static constexpr uint32 ThreadGroupSizeX = 8;
+	static constexpr uint32 ThreadGroupSizeY = 8;
+
+	static void ModifyCompilationEnvironment(const FGlobalShaderPermutationParameters& Parameters, FShaderCompilerEnvironment& OutEnvironment)
+	{
+		FGlobalShader::ModifyCompilationEnvironment(Parameters, OutEnvironment);
+		OutEnvironment.SetDefine(TEXT("THREADGROUP_SIZE_X"), ThreadGroupSizeX);
+		OutEnvironment.SetDefine(TEXT("THREADGROUP_SIZE_Y"), ThreadGroupSizeY);
+	}
+
+	BEGIN_SHADER_PARAMETER_STRUCT(FParameters, )
+		SHADER_PARAMETER_UAV(RWTexture2D<float3>, Texture2DUAV)
+		SHADER_PARAMETER_UAV(RWTexture3D<float3>, Texture3DUAV)
+		SHADER_PARAMETER_UAV(RWTexture2DArray<float3>, TextureArrayUAV)
+		SHADER_PARAMETER(FUintVector3, TextureSize)
+	END_SHADER_PARAMETER_STRUCT()
+};
+
+IMPLEMENT_GLOBAL_SHADER(FFillTexturesCS, "/Plugin/RHITests/Private/TestFillBuffer.usf", "FillTextures", SF_Compute);
+
+bool FRHIReadbackTests::Test_BufferReadback(FRHICommandListImmediate& RHICmdList)
+{
+	const uint32 NumElements = 40;
+	const uint32 TypedBufferSize = NumElements * sizeof(uint32);
+	const uint32 ByteAddressBufferSize = NumElements * sizeof(uint32);
+	const uint32 StructuredBufferSize = NumElements * sizeof(FStructuredBufferElem);
+
+	FBufferRHIRef TypedBuffer, ByteAddressBuffer, StructuredBuffer;
+	FUnorderedAccessViewRHIRef TypedBufferUAV, ByteAddressBufferUAV, StructuredBufferUAV;
+
+	//
+	// Create typed, byte address and structured buffers.
+	//
+	{
+		FRHIResourceCreateInfo TypedBufferCreateInfo(TEXT("TypedBuffer"));
+		EBufferUsageFlags TypedBufferUsage = EBufferUsageFlags::UnorderedAccess | EBufferUsageFlags::SourceCopy;
+		TypedBuffer = RHICmdList.CreateBuffer(TypedBufferSize, TypedBufferUsage, 0, ERHIAccess::UAVCompute, TypedBufferCreateInfo);
+		auto TypedBufferUAVCreateDesc = FRHIViewDesc::CreateBufferUAV()
+			.SetType(FRHIViewDesc::EBufferType::Typed)
+			.SetFormat(PF_R32_UINT)
+			.SetNumElements(NumElements);
+		TypedBufferUAV = RHICmdList.CreateUnorderedAccessView(TypedBuffer, TypedBufferUAVCreateDesc);
+	}
+
+	{
+		FRHIResourceCreateInfo ByteAddressBufferCreateInfo(TEXT("ByteAddressBuffer"));
+		EBufferUsageFlags ByteAddressBufferUsage = EBufferUsageFlags::UnorderedAccess | EBufferUsageFlags::SourceCopy | EBufferUsageFlags::ByteAddressBuffer | EBufferUsageFlags::StructuredBuffer;
+		ByteAddressBuffer = RHICmdList.CreateBuffer(ByteAddressBufferSize, ByteAddressBufferUsage, sizeof(uint32), ERHIAccess::UAVCompute, ByteAddressBufferCreateInfo);
+		auto ByteAddressBufferUAVCreateDesc = FRHIViewDesc::CreateBufferUAV()
+			.SetType(FRHIViewDesc::EBufferType::Raw);
+		ByteAddressBufferUAV = RHICmdList.CreateUnorderedAccessView(ByteAddressBuffer, ByteAddressBufferUAVCreateDesc);
+	}
+
+	{
+		FRHIResourceCreateInfo StructuredBufferCreateInfo(TEXT("StructuredBuffer"));
+		EBufferUsageFlags StructuredBufferUsage = EBufferUsageFlags::UnorderedAccess | EBufferUsageFlags::SourceCopy | EBufferUsageFlags::StructuredBuffer;
+		StructuredBuffer = RHICmdList.CreateBuffer(StructuredBufferSize, StructuredBufferUsage, sizeof(FStructuredBufferElem), ERHIAccess::UAVCompute, StructuredBufferCreateInfo);
+		auto StructuredBufferUAVCreateDesc = FRHIViewDesc::CreateBufferUAV()
+			.SetType(FRHIViewDesc::EBufferType::Structured)
+			.SetStride(sizeof(FStructuredBufferElem))
+			.SetNumElements(NumElements);
+		StructuredBufferUAV = RHICmdList.CreateUnorderedAccessView(StructuredBuffer, StructuredBufferUAVCreateDesc);
+	}
+
+	//
+	// Fill the buffers with data.
+	//
+	{
+		TShaderMapRef<FFillBufferCS> ComputeShader(GetGlobalShaderMap(GMaxRHIFeatureLevel));
+		const uint32 NumThreadGroups = FMath::DivideAndRoundUp(NumElements, FFillBufferCS::ThreadGroupSizeX);
+
+		FFillBufferCS::FParameters Parameters;
+		Parameters.TypedBufferUAV = TypedBufferUAV;
+		Parameters.ByteAddressBufferUAV = ByteAddressBufferUAV;
+		Parameters.StructuredBufferUAV = StructuredBufferUAV;
+		Parameters.NumElements = NumElements;
+
+		FComputeShaderUtils::Dispatch(RHICmdList, ComputeShader, Parameters, FIntVector(NumThreadGroups, 1, 1));
+
+		RHICmdList.Transition({
+			FRHITransitionInfo(TypedBufferUAV, ERHIAccess::UAVCompute, ERHIAccess::CopySrc),
+			FRHITransitionInfo(ByteAddressBufferUAV, ERHIAccess::UAVCompute, ERHIAccess::CopySrc),
+			FRHITransitionInfo(StructuredBufferUAV, ERHIAccess::UAVCompute, ERHIAccess::CopySrc)
+			});
+	}
+
+	// Kick off the readbacks.
+	FRHIGPUBufferReadback TypedBufferReadback(TEXT("TypedBufferReadback"));
+	TypedBufferReadback.EnqueueCopy(RHICmdList, TypedBuffer, TypedBufferSize);
+
+	FRHIGPUBufferReadback ByteAddressBufferReadback(TEXT("ByteAddressBufferReadback"));
+	ByteAddressBufferReadback.EnqueueCopy(RHICmdList, ByteAddressBuffer, ByteAddressBufferSize);
+
+	FRHIGPUBufferReadback StructuredBufferReadback(TEXT("StructuredBufferReadback"));
+	StructuredBufferReadback.EnqueueCopy(RHICmdList, StructuredBuffer, StructuredBufferSize);
+
+	// Sync the GPU. Unfortunately we can't use the fences because not all RHIs implement them yet.
+	RHICmdList.BlockUntilGPUIdle();
+	RHICmdList.FlushResources();
+
+	// Check the contents of the typed buffer.
+	bool bTypedBufferResult = true;
+
+	{
+		const uint32* TypedBufferData = (const uint32*)TypedBufferReadback.Lock(TypedBufferSize);
+		for (uint32 Index = 0; Index < NumElements; ++Index)
+		{
+			const uint32 ExpectedVal = Index;
+			if (TypedBufferData[Index] != ExpectedVal)
+			{
+				UE_LOG(LogRHIUnitTestCommandlet, Error, TEXT("Test failed. \"Typed Buffer Readback\"\n\nIndex: %u; Expected Data: %u; Actual Data: %u\n\n"), Index, ExpectedVal, TypedBufferData[Index]);
+				bTypedBufferResult = false;
+				break;
+			}
+		}
+
+		if (bTypedBufferResult)
+		{
+			UE_LOG(LogRHIUnitTestCommandlet, Display, TEXT("Test passed. \"Typed Buffer Readback\""));
+		}
+
+		TypedBufferReadback.Unlock();
+	}
+
+	// Check the contents of the byte address buffer.
+	bool bByteAddressBufferResult = true;
+
+	{
+		const uint32* ByteAddressBufferData = (const uint32*)ByteAddressBufferReadback.Lock(ByteAddressBufferSize);
+		for (uint32 Index = 0; Index < NumElements; ++Index)
+		{
+			const uint32 ExpectedVal = Index + 1024;
+			if (ByteAddressBufferData[Index] != ExpectedVal)
+			{
+				UE_LOG(LogRHIUnitTestCommandlet, Error, TEXT("Test failed. \"Byte Address Buffer Readback\"\n\nIndex: %u; Expected Data: %u; Actual Data: %u\n\n"), Index, ExpectedVal, ByteAddressBufferData[Index]);
+				bByteAddressBufferResult = false;
+				break;
+			}
+		}
+
+		if (bByteAddressBufferResult)
+		{
+			UE_LOG(LogRHIUnitTestCommandlet, Display, TEXT("Test passed. \"Byte Address Buffer Readback\""));
+		}
+
+		ByteAddressBufferReadback.Unlock();
+	}
+
+	// Check the contents of the structured buffer.
+	bool bStructuredBufferResult = true;
+
+	{
+		const FStructuredBufferElem* StructuredBufferData = (const FStructuredBufferElem*)StructuredBufferReadback.Lock(StructuredBufferSize);
+		for (uint32 Index = 0; Index < NumElements; ++Index)
+		{
+			const uint32 ExpectedIntVal = Index + 2048;
+			const float ExpectedFloatVal = Index + 4096;
+			if (StructuredBufferData[Index].IntVal != ExpectedIntVal || StructuredBufferData[Index].FloatVal != ExpectedFloatVal)
+			{
+				UE_LOG(LogRHIUnitTestCommandlet, Error, TEXT("Test failed. \"Structured Buffer Readback\"\n\nIndex: %u; Expected Int Data: %u; Actual Int Data: %u; Expected Float Data: %.2f; Actual Float Data: %.2f\n\n"),
+					Index, ExpectedIntVal, StructuredBufferData[Index].IntVal, ExpectedFloatVal, StructuredBufferData[Index].FloatVal);
+				bStructuredBufferResult = false;
+				break;
+			}
+		}
+
+		if (bStructuredBufferResult)
+		{
+			UE_LOG(LogRHIUnitTestCommandlet, Display, TEXT("Test passed. \"Structured Buffer Readback\""));
+		}
+
+		StructuredBufferReadback.Unlock();
+	}
+
+	return bTypedBufferResult && bByteAddressBufferResult && bStructuredBufferResult;
+}
+
+bool FRHIReadbackTests::Test_TextureReadback(FRHICommandListImmediate& RHICmdList)
+{
+	const FUintVector3 TextureSize(521, 491, 5);
+
+	FTextureRHIRef Texture2D, Texture3D, TextureArray;
+	FUnorderedAccessViewRHIRef Texture2DUAV, Texture3DUAV, TextureArrayUAV;
+
+	//
+	// Create one texture of each kind.
+	//
+	{
+		const FRHITextureCreateDesc Texture2DDesc =
+			FRHITextureCreateDesc::Create2D(TEXT("Texture2D"))
+			.SetFlags(ETextureCreateFlags::UAV)
+			.SetExtent(TextureSize.X, TextureSize.Y)
+			.SetNumMips(1)
+			.SetFormat(PF_R32_UINT)
+			.SetInitialState(ERHIAccess::UAVCompute);
+
+		Texture2D = RHICreateTexture(Texture2DDesc);
+		auto Texture2DUAVDesc = FRHIViewDesc::CreateTextureUAV()
+			.SetDimensionFromTexture(Texture2D)
+			.SetFormat(PF_R32_UINT);
+
+		Texture2DUAV = RHICmdList.CreateUnorderedAccessView(Texture2D, Texture2DUAVDesc);
+	}
+
+	{
+		const FRHITextureCreateDesc Texture3DDesc =
+			FRHITextureCreateDesc::Create3D(TEXT("Texture3D"))
+			.SetFlags(ETextureCreateFlags::UAV)
+			.SetExtent(TextureSize.X, TextureSize.Y)
+			.SetDepth(TextureSize.Z)
+			.SetNumMips(1)
+			.SetFormat(PF_R32_UINT)
+			.SetInitialState(ERHIAccess::UAVCompute);
+
+		Texture3D = RHICreateTexture(Texture3DDesc);
+		auto Texture3DUAVDesc = FRHIViewDesc::CreateTextureUAV()
+			.SetDimensionFromTexture(Texture3D)
+			.SetFormat(PF_R32_UINT);
+
+		Texture3DUAV = RHICmdList.CreateUnorderedAccessView(Texture3D, Texture3DUAVDesc);
+	}
+
+	{
+		const FRHITextureCreateDesc TextureArrayDesc =
+			FRHITextureCreateDesc::Create2DArray(TEXT("TextureArray"))
+			.SetFlags(ETextureCreateFlags::UAV)
+			.SetExtent(TextureSize.X, TextureSize.Y)
+			.SetArraySize(TextureSize.Z)
+			.SetFormat(PF_R32_UINT)
+			.SetInitialState(ERHIAccess::UAVCompute);
+
+		TextureArray = RHICreateTexture(TextureArrayDesc);
+		auto TextureArrayUAVDesc = FRHIViewDesc::CreateTextureUAV()
+			.SetDimensionFromTexture(TextureArray)
+			.SetFormat(PF_R32_UINT);
+
+		TextureArrayUAV = RHICmdList.CreateUnorderedAccessView(TextureArray, TextureArrayUAVDesc);
+	}
+
+	//
+	// Fill the textures with data.
+	//
+	{
+		TShaderMapRef<FFillTexturesCS> ComputeShader(GetGlobalShaderMap(GMaxRHIFeatureLevel));
+		const uint32 NumThreadGroupsX = FMath::DivideAndRoundUp(TextureSize.X, FFillTexturesCS::ThreadGroupSizeX);
+		const uint32 NumThreadGroupsY = FMath::DivideAndRoundUp(TextureSize.Y, FFillTexturesCS::ThreadGroupSizeY);
+
+		FFillTexturesCS::FParameters Parameters;
+		Parameters.Texture2DUAV = Texture2DUAV;
+		Parameters.Texture3DUAV = Texture3DUAV;
+		Parameters.TextureArrayUAV = TextureArrayUAV;
+		Parameters.TextureSize = TextureSize;
+
+		FComputeShaderUtils::Dispatch(RHICmdList, ComputeShader, Parameters, FIntVector(NumThreadGroupsX, NumThreadGroupsY, 1));
+
+		RHICmdList.Transition({
+			FRHITransitionInfo(Texture2DUAV, ERHIAccess::UAVCompute, ERHIAccess::CopySrc),
+			FRHITransitionInfo(Texture3DUAV, ERHIAccess::UAVCompute, ERHIAccess::CopySrc),
+			FRHITransitionInfo(TextureArrayUAV, ERHIAccess::UAVCompute, ERHIAccess::CopySrc)
+		});
+	}
+
+	// Kick off the readbacks. We'll use an offset to make sure that works correctly. Note that for arrays we only support reading back one slice at a time.
+	FIntVector ReadbackOffset2D(17, 23, 0), ReadbackOffset3D(17, 23, 1), ReadbackOffsetArray(17, 23, 0);
+	FIntVector ReadbackSize2D(293, 173, 1), ReadbackSize3D(293, 173, 3), ReadbackSizeArray(293, 173, 1);
+	const uint32 ReadbackArraySlice = 2;
+
+	FRHIGPUTextureReadback Texture2DReadback(TEXT("Texture2DReadback"));
+	Texture2DReadback.EnqueueCopy(RHICmdList, Texture2D, ReadbackOffset2D, 0, ReadbackSize2D);
+
+	FRHIGPUTextureReadback Texture3DReadback(TEXT("Texture3DReadback"));
+	Texture3DReadback.EnqueueCopy(RHICmdList, Texture3D, ReadbackOffset3D, 0, ReadbackSize3D);
+
+	FRHIGPUTextureReadback TextureArrayReadback(TEXT("TextureArrayReadback"));
+	TextureArrayReadback.EnqueueCopy(RHICmdList, TextureArray, ReadbackOffsetArray, ReadbackArraySlice, ReadbackSizeArray);
+
+	// Sync the GPU. Unfortunately we can't use the fences because not all RHIs implement them yet.
+	RHICmdList.BlockUntilGPUIdle();
+	RHICmdList.FlushResources();
+
+	// Check the contents of the 2D texture.
+	bool bTexture2DResult = true;
+
+	{
+		int32 RowPitchInPixels;
+		const uint32* Texture2DData = (const uint32*)Texture2DReadback.Lock(RowPitchInPixels);
+		const uint32* RowStart = Texture2DData;
+		for (int32 Y = 0; Y < ReadbackSize2D.Y && bTexture2DResult; ++Y)
+		{
+			for (int32 X = 0; X < ReadbackSize2D.X; ++X)
+			{
+				const uint32 ExpectedVal = (Y + ReadbackOffset2D.Y) * TextureSize.X + (X + ReadbackOffset2D.X);
+				if (RowStart[X] != ExpectedVal)
+				{
+					UE_LOG(LogRHIUnitTestCommandlet, Error, TEXT("Test failed. \"Texture 2D Readback\"\n\nCoords: (%u, %u); Expected Data: %u; Actual Data: %u\n\n"), X, Y, ExpectedVal, RowStart[X]);
+					bTexture2DResult = false;
+					break;
+				}
+			}
+
+			RowStart += RowPitchInPixels;
+		}
+
+		if (bTexture2DResult)
+		{
+			UE_LOG(LogRHIUnitTestCommandlet, Display, TEXT("Test passed. \"Texture 2D Readback\""));
+		}
+
+		Texture2DReadback.Unlock();
+	}
+
+	// Check the contents of the 3D texture.
+	bool bTexture3DResult = true;
+
+	{
+		int32 RowPitchInPixels, BufferHeight;
+		const uint32* Texture3DData = (const uint32*)Texture3DReadback.Lock(RowPitchInPixels, &BufferHeight);
+		const uint32* SliceStart = Texture3DData;
+		for (int32 Z = 0; Z < ReadbackSize3D.Z && bTexture3DResult; ++Z)
+		{
+			const uint32* RowStart = SliceStart;
+			for (int32 Y = 0; Y < ReadbackSize3D.Y && bTexture3DResult; ++Y)
+			{
+				for (int32 X = 0; X < ReadbackSize3D.X; ++X)
+				{
+					const uint32 ExpectedVal = (Y + ReadbackOffset3D.Y) * TextureSize.X + (X + ReadbackOffset3D.X) + (Z + ReadbackOffset3D.Z + 10) * 1024;
+					if (RowStart[X] != ExpectedVal)
+					{
+						UE_LOG(LogRHIUnitTestCommandlet, Error, TEXT("Test failed. \"Texture 3D Readback\"\n\nCoords: (%u, %u, %u); Expected Data: %u; Actual Data: %u\n\n"), X, Y, Z, ExpectedVal, RowStart[X]);
+						bTexture3DResult = false;
+						break;
+					}
+				}
+
+				RowStart += RowPitchInPixels;
+			}
+
+			SliceStart += BufferHeight * RowPitchInPixels;
+		}
+
+		if (bTexture3DResult)
+		{
+			UE_LOG(LogRHIUnitTestCommandlet, Display, TEXT("Test passed. \"Texture 3D Readback\""));
+		}
+
+		Texture3DReadback.Unlock();
+	}
+
+	// Check the contents of the array texture.
+	bool bTextureArrayResult = true;
+
+	{
+		int32 RowPitchInPixels;
+		const uint32* TextureArrayData = (const uint32*)TextureArrayReadback.Lock(RowPitchInPixels);
+		const uint32* RowStart = TextureArrayData;
+		for (int32 Y = 0; Y < ReadbackSizeArray.Y && bTextureArrayResult; ++Y)
+		{
+			for (int32 X = 0; X < ReadbackSizeArray.X; ++X)
+			{
+				const uint32 ExpectedVal = (Y + ReadbackOffsetArray.Y) * TextureSize.X + (X + ReadbackOffsetArray.X) + (ReadbackArraySlice + 10 + TextureSize.Z) * 1024;
+				if (RowStart[X] != ExpectedVal)
+				{
+					UE_LOG(LogRHIUnitTestCommandlet, Error, TEXT("Test failed. \"Texture Array Readback\"\n\nCoords: (%u, %u); Slice: %u; Expected Data: %u; Actual Data: %u\n\n"), X, Y, ReadbackArraySlice, ExpectedVal, RowStart[X]);
+					bTextureArrayResult = false;
+					break;
+				}
+			}
+
+			RowStart += RowPitchInPixels;
+		}
+
+		if (bTextureArrayResult)
+		{
+			UE_LOG(LogRHIUnitTestCommandlet, Display, TEXT("Test passed. \"Texture Array Readback\""));
+		}
+
+		TextureArrayReadback.Unlock();
+	}
+
+	return bTexture2DResult && bTexture3DResult && bTextureArrayResult;
+}
diff --git a/Engine/Plugins/Tests/RHITests/Source/RHITests/Private/RHITests.spec.cpp b/Engine/Plugins/Tests/RHITests/Source/RHITests/Private/RHITests.spec.cpp
--- a/Engine/Plugins/Tests/RHITests/Source/RHITests/Private/RHITests.spec.cpp
+++ b/Engine/Plugins/Tests/RHITests/Source/RHITests/Private/RHITests.spec.cpp
@@ -7,6 +7,7 @@
 #include "RHIBufferTests.h"
 #include "RHITextureTests.h"
 #include "RHIDrawTests.h"
+#include "RHIReadbackTests.h"
 
 BEGIN_DEFINE_SPEC(FAutomationRHITest, "Rendering.RHI", EAutomationTestFlags::EngineFilter | EAutomationTestFlags::ApplicationContextMask | EAutomationTestFlags::NonNullRHI)
 END_DEFINE_SPEC(FAutomationRHITest)
@@ -84,6 +85,21 @@ void FAutomationRHITest::Define()
 		});
 	});
 
+	Describe("Test RHI Readback", [this]
+	{
+		It("Buffer Readback", [this]()
+		{
+			bool bResult = RunOnRenderThreadSynchronous(FRHIReadbackTests::Test_BufferReadback);
+			TestEqual("BufferReadback", bResult, 1);
+		});
+
+		It("Texture Readback", [this]()
+		{
+			bool bResult = RunOnRenderThreadSynchronous(FRHIReadbackTests::Test_TextureReadback);
+			TestEqual("TextureReadback", bResult, 1);
+		});
+	});
+
 	Describe("Test RHI Create Buffer Parallel", [this]
 	{
 		It("RHICreateBuffer_Parallel", [this]()
diff --git a/Engine/Plugins/Tests/RHITests/Source/RHITests/Private/RHITestsModule.cpp b/Engine/Plugins/Tests/RHITests/Source/RHITests/Private/RHITestsModule.cpp
--- a/Engine/Plugins/Tests/RHITests/Source/RHITests/Private/RHITestsModule.cpp
+++ b/Engine/Plugins/Tests/RHITests/Source/RHITests/Private/RHITestsModule.cpp
@@ -8,6 +8,7 @@
 #include "RHIBufferTests.h"
 #include "RHITextureTests.h"
 #include "RHIDrawTests.h"
+#include "RHIReadbackTests.h"
 
 #define LOCTEXT_NAMESPACE "FRHITestsModule"
 
@@ -54,6 +55,14 @@ static bool RunTests_RenderThread(FRHICommandListImmediate& RHICmdList)
 		RUN_TEST(FRHITextureTests::Test_MultipleLockTexture2D(RHICmdList));
 	}
 
+	// ------------------------------------------------
+	// Readback
+	// ------------------------------------------------
+	{
+		RUN_TEST(FRHIReadbackTests::Test_BufferReadback(RHICmdList));
+		RUN_TEST(FRHIReadbackTests::Test_TextureReadback(RHICmdList));
+	}
+
 	// ------------------------------------------------
 	// Buffer Operations
 	// ------------------------------------------------
diff --git a/Engine/Plugins/Tests/RHITests/Source/RHITests/Public/RHIReadbackTests.h b/Engine/Plugins/Tests/RHITests/Source/RHITests/Public/RHIReadbackTests.h
new file mode 100644
--- /dev/null
+++ b/Engine/Plugins/Tests/RHITests/Source/RHITests/Public/RHIReadbackTests.h
@@ -0,0 +1,12 @@
+// Copyright Epic Games, Inc. All Rights Reserved.
+
+#pragma once
+
+#include "RHITestsCommon.h"
+
+class FRHIReadbackTests
+{
+public:
+	static bool Test_BufferReadback(FRHICommandListImmediate& RHICmdList);
+	static bool Test_TextureReadback(FRHICommandListImmediate& RHICmdList);
+};
diff --git a/Engine/Source/Runtime/D3D12RHI/Private/D3D12RenderTarget.cpp b/Engine/Source/Runtime/D3D12RHI/Private/D3D12RenderTarget.cpp
--- a/Engine/Source/Runtime/D3D12RHI/Private/D3D12RenderTarget.cpp
+++ b/Engine/Source/Runtime/D3D12RHI/Private/D3D12RenderTarget.cpp
@@ -1130,12 +1130,12 @@ void FD3D12DynamicRHI::ReadSurfaceDataMSAARaw(FRHITexture* TextureRHI, FIntRect
 
 void FD3D12DynamicRHI::RHIMapStagingSurface(FRHITexture* TextureRHI, FRHIGPUFence* FenceRHI, void*& OutData, int32& OutWidth, int32& OutHeight, uint32 GPUIndex)
 {	
-	FD3D12Texture* DestTexture2D = ResourceCast(TextureRHI->GetTexture2D(), GPUIndex);
+	FD3D12Texture* DestTexture = ResourceCast(TextureRHI, GPUIndex);
 
-	check(DestTexture2D);
-	FD3D12Resource* Texture = DestTexture2D->GetResource();
+	check(DestTexture);
+	FD3D12Resource* Texture = DestTexture->GetResource();
 
-	DXGI_FORMAT Format = (DXGI_FORMAT)GPixelFormats[DestTexture2D->GetFormat()].PlatformFormat;
+	DXGI_FORMAT Format = (DXGI_FORMAT)GPixelFormats[DestTexture->GetFormat()].PlatformFormat;
 
 	uint32 BytesPerPixel = ComputeBytesPerPixel(Format);
 
@@ -1153,7 +1153,7 @@ void FD3D12DynamicRHI::RHIMapStagingSurface(FRHITexture* TextureRHI, FRHIGPUFenc
 	);
 
 	D3D12_PLACED_SUBRESOURCE_FOOTPRINT ReadBackHeapDesc;
-	DestTexture2D->GetReadBackHeapDesc(ReadBackHeapDesc, 0);
+	DestTexture->GetReadBackHeapDesc(ReadBackHeapDesc, 0);
 	OutData = pData;
 	OutWidth = ReadBackHeapDesc.Footprint.RowPitch / BytesPerPixel;
 	OutHeight = ReadBackHeapDesc.Footprint.Height;
@@ -1170,10 +1170,9 @@ void FD3D12DynamicRHI::RHIMapStagingSurface(FRHITexture* TextureRHI, FRHIGPUFenc
 
 void FD3D12DynamicRHI::RHIUnmapStagingSurface(FRHITexture* TextureRHI, uint32 GPUIndex)
 {	
-	FD3D12Texture* DestTexture2D = ResourceCast(TextureRHI->GetTexture2D(), GPUIndex);
+	FD3D12Texture* DestTexture = ResourceCast(TextureRHI, GPUIndex);
 
-	check(DestTexture2D);
-	ID3D12Resource* Texture = DestTexture2D->GetResource()->GetResource();
+	ID3D12Resource* Texture = DestTexture->GetResource()->GetResource();
 
 	Texture->Unmap(0, nullptr);
 }
diff --git a/Engine/Source/Runtime/D3D12RHI/Private/D3D12Texture.cpp b/Engine/Source/Runtime/D3D12RHI/Private/D3D12Texture.cpp
--- a/Engine/Source/Runtime/D3D12RHI/Private/D3D12Texture.cpp
+++ b/Engine/Source/Runtime/D3D12RHI/Private/D3D12Texture.cpp
@@ -944,6 +944,20 @@ FD3D12Texture* FD3D12DynamicRHI::CreateD3D12Texture(const FRHITextureCreateDesc&
 			ResourceAllocator->AllocateTexture(Device->GetGPUIndex(), HeapType, ResourceDesc, (EPixelFormat)CreateDesc.Format, ED3D12ResourceStateMode::Default, CreateState, ClearValuePtr, CreateDesc.DebugName, Location);
 			Location.SetOwner(NewTexture);
 		}
+		else if(EnumHasAnyFlags(CreateDesc.Flags, TexCreate_CPUReadback))
+		{
+			uint64 Size = 0;
+			uint32 NumSubResources = ResourceDesc.MipLevels;
+			if (CreateDesc.IsTextureArray())
+			{
+				NumSubResources *= ResourceDesc.DepthOrArraySize;
+			}
+			Device->GetDevice()->GetCopyableFootprints(&ResourceDesc, 0, NumSubResources, 0, nullptr, nullptr, nullptr, &Size);
+
+			FD3D12Resource* Resource = nullptr;
+			VERIFYD3D12CREATETEXTURERESULT(Adapter->CreateBuffer(D3D12_HEAP_TYPE_READBACK, Device->GetGPUMask(), Device->GetVisibilityMask(), Size, &Resource, CreateDesc.DebugName), ResourceDesc, Device->GetDevice());
+			Location.AsStandAlone(Resource);
+		}
 		else if (CreateDesc.IsTexture3D())
 		{
 			VERIFYD3D12CREATETEXTURERESULT(Device->GetTextureAllocator().AllocateTexture(
diff --git a/Engine/Source/Runtime/RHI/Private/RHIGPUReadback.cpp b/Engine/Source/Runtime/RHI/Private/RHIGPUReadback.cpp
--- a/Engine/Source/Runtime/RHI/Private/RHIGPUReadback.cpp
+++ b/Engine/Source/Runtime/RHI/Private/RHIGPUReadback.cpp
@@ -94,8 +94,6 @@ void* FRHIGPUBufferReadback::Lock(uint32 NumBytes)
 	if (DestinationStagingBuffers[GPUIndex])
 	{
 		LastLockGPUIndex = GPUIndex;
-
-		ensure(Fence->Poll());
 		return RHILockStagingBuffer(DestinationStagingBuffers[GPUIndex], Fence.GetReference(), 0, NumBytes);
 	}
 	else
@@ -132,14 +130,13 @@ FRHIGPUTextureReadback::FRHIGPUTextureReadback(FName RequestName) : FRHIGPUMemor
 {
 }
 
-void FRHIGPUTextureReadback::EnqueueCopy(FRHICommandList& RHICmdList, FRHITexture* SourceTexture, FResolveRect Rect)
+void FRHIGPUTextureReadback::EnqueueCopy(FRHICommandList& RHICmdList, FRHITexture* SourceTexture, const FIntVector& SourcePosition, uint32 SourceSlice, const FIntVector& Size)
 {
 	Fence->Clear();
 	LastCopyGPUMask = RHICmdList.GetGPUMask();
 
 	if (SourceTexture)
 	{
-		check(SourceTexture->GetTexture2D());
 		check(!SourceTexture->IsMultisampled());
 
 		for (uint32 GPUIndex : LastCopyGPUMask)
@@ -147,12 +144,14 @@ void FRHIGPUTextureReadback::EnqueueCopy(FRHICommandList& RHICmdList, FRHITextur
 			SCOPED_GPU_MASK(RHICmdList, FRHIGPUMask::FromIndex(GPUIndex));
 
 			// Assume for now that every enqueue happens on a texture of the same format and size (when reused).
-			if (!DestinationStagingTextures[GPUIndex])
+			if (!DestinationStagingTextures[GPUIndex] || DestinationStagingTextures[GPUIndex]->GetDesc().Dimension != SourceTexture->GetDesc().Dimension)
 			{
 				FIntVector StagingTextureSize;
-				if (Rect.IsValid())
+
+				// Passing 0 or negative for size means read back the entire texture.
+				if (Size.X > 0 && Size.Y > 0)
 				{
-					StagingTextureSize = FIntVector(Rect.X2 - Rect.X1, Rect.Y2 - Rect.Y1, 1);
+					StagingTextureSize = Size;
 				}
 				else
 				{
@@ -161,10 +160,22 @@ void FRHIGPUTextureReadback::EnqueueCopy(FRHICommandList& RHICmdList, FRHITextur
 
 				FString FenceName = Fence->GetFName().ToString();
 
-				const FRHITextureCreateDesc Desc =
-					FRHITextureCreateDesc::Create2D(*FenceName, StagingTextureSize.X, StagingTextureSize.Y, SourceTexture->GetFormat())
-					.SetFlags(ETextureCreateFlags::CPUReadback | ETextureCreateFlags::HideInVisualizeTexture)
-					.SetInitialState(ERHIAccess::CopyDest);
+				FRHITextureCreateDesc Desc = FRHITextureCreateDesc(SourceTexture->GetDesc(), ERHIAccess::CopyDest, *FenceName)
+					.SetExtent(StagingTextureSize.X, StagingTextureSize.Y)
+					.SetFlags(ETextureCreateFlags::CPUReadback | ETextureCreateFlags::HideInVisualizeTexture);
+
+				switch (Desc.Dimension)
+				{
+					case ETextureDimension::Texture2DArray:
+					case ETextureDimension::TextureCubeArray:
+						ensureMsgf(Size.Z <= 1, TEXT("Readback for texture arrays supports only one slice at a time. Texture Name: %s, SourcePosition: (%d, %d, %d), SourceSlice: %u, Size: (%d, %d, %d)."),
+							*SourceTexture->GetName().ToString(), SourcePosition.X, SourcePosition.Y, SourcePosition.Z, SourceSlice, Size.X, Size.Y, Size.Z);
+						Desc.SetArraySize(1);
+						break;
+					case ETextureDimension::Texture3D:
+						Desc.SetDepth(StagingTextureSize.Z);
+						break;
+				}
 
 				DestinationStagingTextures[GPUIndex] = RHICreateTexture(Desc);
 			}
@@ -172,14 +183,24 @@ void FRHIGPUTextureReadback::EnqueueCopy(FRHICommandList& RHICmdList, FRHITextur
 			{
 				RHICmdList.Transition(FRHITransitionInfo(DestinationStagingTextures[GPUIndex], ERHIAccess::CPURead, ERHIAccess::CopyDest));
 			}
+			
 			FRHICopyTextureInfo CopyInfo;
 
-			if (Rect.IsValid())
+			// Make sure we're not passing negative coordinates or size.
+			if (SourcePosition.X >= 0 && SourcePosition.Y >= 0)
 			{
-				CopyInfo.SourcePosition = FIntVector(Rect.X1, Rect.Y1, 0);
-				CopyInfo.Size = FIntVector(Rect.X2 - Rect.X1, Rect.Y2 - Rect.Y1, 1);
+				CopyInfo.SourcePosition = SourcePosition;
 			}
 
+			if (Size.X > 0 && Size.Y > 0)
+			{
+				CopyInfo.Size = Size;
+			}
+
+			CopyInfo.SourceSliceIndex = SourceSlice;
+			CopyInfo.DestSliceIndex = 0;
+			CopyInfo.NumSlices = 1;
+
 			RHICmdList.CopyTexture(SourceTexture, DestinationStagingTextures[GPUIndex], CopyInfo);
 			RHICmdList.Transition(FRHITransitionInfo(DestinationStagingTextures[GPUIndex], ERHIAccess::CopyDest, ERHIAccess::CPURead));
 			RHICmdList.WriteGPUFence(Fence);
diff --git a/Engine/Source/Runtime/RHI/Public/RHIGPUReadback.h b/Engine/Source/Runtime/RHI/Public/RHIGPUReadback.h
--- a/Engine/Source/Runtime/RHI/Public/RHIGPUReadback.h
+++ b/Engine/Source/Runtime/RHI/Public/RHIGPUReadback.h
@@ -51,11 +51,29 @@ public:
 		unimplemented();
 	}
 
-	virtual void EnqueueCopy(FRHICommandList& RHICmdList, FRHITexture* SourceTexture, FResolveRect Rect = FResolveRect())
+	virtual void EnqueueCopy(FRHICommandList& RHICmdList, FRHITexture* SourceTexture, const FIntVector& SourcePosition, uint32 SourceSlice, const FIntVector& Size)
 	{
 		unimplemented();
 	}
 
+	void EnqueueCopy(FRHICommandList& RHICmdList, FRHITexture* SourceTexture, FResolveRect Rect = FResolveRect())
+	{
+		FIntVector SourcePosition, Size;
+
+		if (Rect.IsValid())
+		{
+			SourcePosition = FIntVector(Rect.X1, Rect.Y1, 0);
+			Size = FIntVector(Rect.X2 - Rect.X1, Rect.Y2 - Rect.Y1, 1);
+		}
+		else
+		{
+			SourcePosition = FIntVector::ZeroValue;
+			Size = FIntVector::ZeroValue;
+		}
+		
+		EnqueueCopy(RHICmdList, SourceTexture, SourcePosition, 0, Size);
+	}
+
 	/**
 	 * Returns the CPU accessible pointer that backs this staging buffer.
 	 * @param NumBytes The maximum number of bytes the host will read from this pointer.
@@ -114,7 +132,9 @@ class FRHIGPUTextureReadback final : public FRHIGPUMemoryReadback
 public:
 	RHI_API FRHIGPUTextureReadback(FName RequestName);
 
-	RHI_API void EnqueueCopy(FRHICommandList& RHICmdList, FRHITexture* SourceTexture, FResolveRect Rect = FResolveRect()) override;
+	using FRHIGPUMemoryReadback::EnqueueCopy;
+
+	RHI_API void EnqueueCopy(FRHICommandList& RHICmdList, FRHITexture* SourceTexture, const FIntVector& SourcePosition, uint32 SourceSlice, const FIntVector& Size) override;
 
 	UE_DEPRECATED(5.0, "Use FRHIGPUTextureReadback::Lock( int32& OutRowPitchInPixels) instead.")
 	RHI_API void* Lock(uint32 NumBytes) override;
diff --git a/Engine/Source/Runtime/VulkanRHI/Private/VulkanRenderTarget.cpp b/Engine/Source/Runtime/VulkanRHI/Private/VulkanRenderTarget.cpp
--- a/Engine/Source/Runtime/VulkanRHI/Private/VulkanRenderTarget.cpp
+++ b/Engine/Source/Runtime/VulkanRHI/Private/VulkanRenderTarget.cpp
@@ -254,7 +254,6 @@ void FVulkanDynamicRHI::RHIReadSurfaceData(FRHITexture* TextureRHI, FIntRect Rec
 
 void FVulkanDynamicRHI::RHIMapStagingSurface(FRHITexture* TextureRHI, FRHIGPUFence* FenceRHI, void*& OutData, int32& OutWidth, int32& OutHeight, uint32 GPUIndex)
 {
-	check(TextureRHI->GetDesc().Dimension == ETextureDimension::Texture2D);
 	FVulkanTexture* Texture = ResourceCast(TextureRHI);
 
 	if (FenceRHI && !FenceRHI->Poll())
diff --git a/Engine/Source/Runtime/VulkanRHI/Private/VulkanTexture.cpp b/Engine/Source/Runtime/VulkanRHI/Private/VulkanTexture.cpp
--- a/Engine/Source/Runtime/VulkanRHI/Private/VulkanTexture.cpp
+++ b/Engine/Source/Runtime/VulkanRHI/Private/VulkanTexture.cpp
@@ -801,7 +801,7 @@ void FVulkanTexture::GetMipSize(uint32 MipIndex, uint32& MipBytes)
 	}
 
 	// Size in bytes
-	MipBytes = NumBlocksX * NumBlocksY * BlockBytes;
+	MipBytes = NumBlocksX * NumBlocksY * BlockBytes * Desc.Depth;
 /*
 #if VULKAN_HAS_DEBUGGING_ENABLED
 	VkImageSubresource SubResource;
@@ -1454,7 +1454,6 @@ FVulkanTexture::FVulkanTexture(FVulkanDevice& InDevice, const FRHITextureCreateD
 		if (EnumHasAnyFlags(InCreateDesc.Flags, TexCreate_CPUReadback))
 		{
 			check(InCreateDesc.NumSamples == 1); //not implemented
-			check(InCreateDesc.Depth == 1);      //not implemented
 			check(InCreateDesc.ArraySize == 1);  //not implemented
 
 			CpuReadbackBuffer = new FVulkanCpuReadbackBuffer;
diff --git a/Engine/Source/Runtime/Windows/D3D11RHI/Private/D3D11RenderTarget.cpp b/Engine/Source/Runtime/Windows/D3D11RHI/Private/D3D11RenderTarget.cpp
--- a/Engine/Source/Runtime/Windows/D3D11RHI/Private/D3D11RenderTarget.cpp
+++ b/Engine/Source/Runtime/Windows/D3D11RHI/Private/D3D11RenderTarget.cpp
@@ -655,14 +655,13 @@ void FD3D11DynamicRHI::ReadSurfaceDataMSAARaw(FRHITexture* TextureRHI,FIntRect I
 
 void FD3D11DynamicRHI::RHIMapStagingSurface(FRHITexture* TextureRHI, FRHIGPUFence* FenceRHI, void*& OutData, int32& OutWidth, int32& OutHeight, uint32 GPUIndex)
 {
-	ID3D11Texture2D* Texture = ResourceCast(TextureRHI)->GetD3D11Texture2D();
-	
-	D3D11_TEXTURE2D_DESC TextureDesc;
-	Texture->GetDesc(&TextureDesc);
-	uint32 BytesPerPixel = ComputeBytesPerPixel(TextureDesc.Format);
+	ID3D11Resource* Resource = ResourceCast(TextureRHI)->GetResource();
+
+	DXGI_FORMAT Format = (DXGI_FORMAT)GPixelFormats[TextureRHI->GetDesc().Format].PlatformFormat;
+	uint32 BytesPerPixel = ComputeBytesPerPixel(Format);
 
 	D3D11_MAPPED_SUBRESOURCE LockedRect;
-	VERIFYD3D11RESULT_EX(Direct3DDeviceIMContext->Map(Texture,0,D3D11_MAP_READ,0,&LockedRect), Direct3DDevice);
+	VERIFYD3D11RESULT_EX(Direct3DDeviceIMContext->Map(Resource, 0, D3D11_MAP_READ, 0, &LockedRect), Direct3DDevice);
 
 	OutData = LockedRect.pData;
 	OutWidth = LockedRect.RowPitch / BytesPerPixel;
@@ -673,9 +672,8 @@ void FD3D11DynamicRHI::RHIMapStagingSurface(FRHITexture* TextureRHI, FRHIGPUFenc
 
 void FD3D11DynamicRHI::RHIUnmapStagingSurface(FRHITexture* TextureRHI, uint32 GPUIndex)
 {
-	ID3D11Texture2D* Texture = ResourceCast(TextureRHI)->GetD3D11Texture2D();
-
-	Direct3DDeviceIMContext->Unmap(Texture,0);
+	ID3D11Resource* Resource = ResourceCast(TextureRHI)->GetResource();
+	Direct3DDeviceIMContext->Unmap(Resource, 0);
 }
 
 void FD3D11DynamicRHI::RHIReadSurfaceFloatData(FRHITexture* TextureRHI,FIntRect InRect,TArray<FFloat16Color>& OutData,ECubeFace CubeFace,int32 ArrayIndex,int32 MipIndex)
diff --git a/Engine/Source/Runtime/Windows/D3D11RHI/Private/D3D11Texture.cpp b/Engine/Source/Runtime/Windows/D3D11RHI/Private/D3D11Texture.cpp
--- a/Engine/Source/Runtime/Windows/D3D11RHI/Private/D3D11Texture.cpp
+++ b/Engine/Source/Runtime/Windows/D3D11RHI/Private/D3D11Texture.cpp
@@ -759,6 +759,23 @@ FD3D11Texture* FD3D11DynamicRHI::CreateD3D11Texture3D(FRHITextureCreateDesc cons
 		bCreateRTV = true;
 	}
 
+	bool bCreateShaderResource = true;
+
+	if (EnumHasAnyFlags(Flags, TexCreate_CPUReadback))
+	{
+		check(!EnumHasAnyFlags(Flags, TexCreate_RenderTargetable | TexCreate_DepthStencilTargetable | TexCreate_ShaderResource));
+
+		TextureDesc.CPUAccessFlags = D3D11_CPU_ACCESS_READ;
+		TextureDesc.Usage = D3D11_USAGE_STAGING;
+		TextureDesc.BindFlags = 0;
+		bCreateShaderResource = false;
+	}
+
+	if (EnumHasAnyFlags(Flags, TexCreate_DisableSRVCreation))
+	{
+		bCreateShaderResource = false;
+	}
+
 	// Set up the texture bind flags.
 	check(!EnumHasAnyFlags(Flags, TexCreate_DepthStencilTargetable | TexCreate_ResolveTargetable));
 
@@ -803,6 +820,7 @@ FD3D11Texture* FD3D11DynamicRHI::CreateD3D11Texture3D(FRHITextureCreateDesc cons
 
 	// Create a shader resource view for the texture.
 	TRefCountPtr<ID3D11ShaderResourceView> ShaderResourceView;
+	if(bCreateShaderResource)
 	{
 		D3D11_SHADER_RESOURCE_VIEW_DESC SRVDesc;
 		SRVDesc.Format = PlatformShaderResourceFormat;
