Skip PF_R9G9B9EXP5 and PF_P010 when running the RHI texture tests.

#rnx

diff --git a/Engine/Plugins/Tests/RHITests/Source/RHITests/Public/RHITextureTests.h b/Engine/Plugins/Tests/RHITests/Source/RHITests/Public/RHITextureTests.h
--- a/Engine/Plugins/Tests/RHITests/Source/RHITests/Public/RHITextureTests.h
+++ b/Engine/Plugins/Tests/RHITests/Source/RHITests/Public/RHITextureTests.h
@@ -721,6 +721,8 @@ public:
 			case PF_ASTC_8x8_NORM_RG:
 			case PF_ASTC_10x10_NORM_RG:
 			case PF_ASTC_12x12_NORM_RG:
+			case PF_R9G9B9EXP5:
+			case PF_P010:
 				break;
 
 			case PF_G8:
