Fixed missing deprecated RHICreateBuffer.

diff --git a/Engine/Plugins/Tests/RHITests/Source/RHITests/Public/RHIBufferTests.h b/Engine/Plugins/Tests/RHITests/Source/RHITests/Public/RHIBufferTests.h
--- a/Engine/Plugins/Tests/RHITests/Source/RHITests/Public/RHIBufferTests.h
+++ b/Engine/Plugins/Tests/RHITests/Source/RHITests/Public/RHIBufferTests.h
@@ -169,7 +169,7 @@ private:
 		}
 
 		FRHIResourceCreateInfo Info(*TestName);
-		FBufferRHIRef Buffer = RHICreateBuffer(BufferSize, BufferFlags | BUF_ShaderResource | BUF_UnorderedAccess | BUF_SourceCopy, Stride, ERHIAccess::UAVMask, Info);
+		FBufferRHIRef Buffer = RHICmdList.CreateBuffer(BufferSize, BufferFlags | BUF_ShaderResource | BUF_UnorderedAccess | BUF_SourceCopy, Stride, ERHIAccess::UAVMask, Info);
 
 		auto UAVCreateDesc = FRHIViewDesc::CreateBufferUAV();
 		auto SRVCreateDesc = FRHIViewDesc::CreateBufferSRV();
