Remove BUF_ImmediateCommit in favor of only having an explicit buffer commit operation

diff --git a/Engine/Plugins/Tests/RHITests/Source/RHITests/Private/RHIReservedResourceTests.cpp b/Engine/Plugins/Tests/RHITests/Source/RHITests/Private/RHIReservedResourceTests.cpp
--- a/Engine/Plugins/Tests/RHITests/Source/RHITests/Private/RHIReservedResourceTests.cpp
+++ b/Engine/Plugins/Tests/RHITests/Source/RHITests/Private/RHIReservedResourceTests.cpp
@@ -93,26 +93,31 @@ bool FRHIReservedResourceTests::Test_ReservedResource_CreateBuffer(FRHICommandLi
 		return true;
 	}
 
-	const EBufferUsageFlags CommonFlags = BUF_ReservedResource | BUF_ImmediateCommit;
+	// Simply try to create reserved buffers of different types and sizes to see if we hit any unexpected paths in the RHI
 
 	{
 		FRHIResourceCreateInfo CreateInfo(TEXT("TestSmallReservedVertexBuffer"));
-		FBufferRHIRef Buffer = RHICmdList.CreateBuffer(32768, CommonFlags | BUF_VertexBuffer, 4, ERHIAccess::CopyDest, CreateInfo);
+		FBufferRHIRef Buffer = RHICmdList.CreateBuffer(32768, BUF_ReservedResource | BUF_VertexBuffer, 4, ERHIAccess::CopyDest, CreateInfo);
+	}
+
+	{
+		FRHIResourceCreateInfo CreateInfo(TEXT("TestSmallReservedUAV"));
+		FBufferRHIRef Buffer = RHICmdList.CreateBuffer(32768, BUF_ReservedResource | BUF_UnorderedAccess | BUF_ShaderResource, 4, ERHIAccess::UAVGraphics, CreateInfo);
 	}
 
 	{
 		FRHIResourceCreateInfo CreateInfo(TEXT("TestReservedVertexBuffer"));
-		FBufferRHIRef Buffer = RHICmdList.CreateBuffer(32 * 1024 * 1024, CommonFlags | BUF_VertexBuffer, 4, ERHIAccess::CopyDest, CreateInfo);
+		FBufferRHIRef Buffer = RHICmdList.CreateBuffer(32 * 1024 * 1024, BUF_ReservedResource | BUF_VertexBuffer, 4, ERHIAccess::CopyDest, CreateInfo);
 	}
 
 	{
 		FRHIResourceCreateInfo CreateInfo(TEXT("TestReservedAccelerationStructureBuffer"));
-		FBufferRHIRef Buffer = RHICmdList.CreateBuffer(32 * 1024 * 1024, CommonFlags | BUF_AccelerationStructure, 4, ERHIAccess::BVHWrite, CreateInfo);
+		FBufferRHIRef Buffer = RHICmdList.CreateBuffer(32 * 1024 * 1024, BUF_ReservedResource | BUF_AccelerationStructure, 4, ERHIAccess::BVHWrite, CreateInfo);
 	}
 
 	{
 		FRHIResourceCreateInfo CreateInfo(TEXT("TestReservedRayTracingScratchBuffer"));
-		FBufferRHIRef Buffer = RHICmdList.CreateBuffer(32 * 1024 * 1024, CommonFlags | BUF_RayTracingScratch, 4, ERHIAccess::UAVCompute, CreateInfo);
+		FBufferRHIRef Buffer = RHICmdList.CreateBuffer(32 * 1024 * 1024, BUF_ReservedResource | BUF_RayTracingScratch, 4, ERHIAccess::UAVCompute, CreateInfo);
 	}
 
 	return true;
@@ -185,7 +190,8 @@ bool FRHIReservedResourceTests::Test_ReservedResource_CommitBuffer(FRHICommandLi
 				}
 			}
 
-			uint64 ExpectedTailValue = 0;
+			// We follow the D3D convention for unmapped page access: writes are no-op, reads return 0
+			const uint64 ExpectedTailValue = 0;
 			for (uint32 i = CommittedSizeInElements; i < TotalSizeInElements; ++i)
 			{
 				if (BufferData[i] != ExpectedTailValue)
diff --git a/Engine/Source/Runtime/D3D12RHI/Private/D3D12Allocation.cpp b/Engine/Source/Runtime/D3D12RHI/Private/D3D12Allocation.cpp
--- a/Engine/Source/Runtime/D3D12RHI/Private/D3D12Allocation.cpp
+++ b/Engine/Source/Runtime/D3D12RHI/Private/D3D12Allocation.cpp
@@ -1537,14 +1537,6 @@ void FD3D12DefaultBufferAllocator::AllocDefaultResource(D3D12_HEAP_TYPE InHeapTy
 
 		ResourceLocation.AsStandAlone(NewResource, InResourceDesc.Width);
 
-		if (EnumHasAnyFlags(InBufferUsage, BUF_ImmediateCommit))
-		{
-			// NOTE: Accessing the queue from this thread is OK, as D3D12 runtime acquires a lock around all command queue APIs.
-			// https://microsoft.github.io/DirectX-Specs/d3d/CPUEfficiency.html#threading
-			FD3D12Queue& Queue = GetParentDevice()->GetQueue(ED3D12QueueType::Direct);
-			NewResource->CommitReservedResource(Queue.D3DCommandQueue, UINT64_MAX /*commit entire resource*/);
-		}
-
 		return;
 	}
 
diff --git a/Engine/Source/Runtime/D3D12RHI/Private/D3D12Buffer.cpp b/Engine/Source/Runtime/D3D12RHI/Private/D3D12Buffer.cpp
--- a/Engine/Source/Runtime/D3D12RHI/Private/D3D12Buffer.cpp
+++ b/Engine/Source/Runtime/D3D12RHI/Private/D3D12Buffer.cpp
@@ -563,7 +563,7 @@ FD3D12Buffer* FD3D12DynamicRHI::CreateD3D12Buffer(class FRHICommandListBase* RHI
 
 	if (EnumHasAnyFlags(BufferDesc.Usage, BUF_ReservedResource))
 	{
-		checkf(!bHasInitialData || EnumHasAllFlags(BufferDesc.Usage, BUF_ImmediateCommit), TEXT("Reserved resources with initial data must be committed at creation"));
+		checkf(!bHasInitialData, TEXT("Reserved resources may not have initial data"));
 		checkf(!bIsDynamic, TEXT("Reserved resources may not be dynamic"));
 		checkf(!ResourceAllocator, TEXT("Reserved resources may not use a custom resource allocator"));
 	}
diff --git a/Engine/Source/Runtime/RHI/Public/RHIDefinitions.h b/Engine/Source/Runtime/RHI/Public/RHIDefinitions.h
--- a/Engine/Source/Runtime/RHI/Public/RHIDefinitions.h
+++ b/Engine/Source/Runtime/RHI/Public/RHIDefinitions.h
@@ -783,7 +783,7 @@ enum class EBufferUsageFlags : uint32
 	/**
 	 * Buffer contains opaque ray tracing acceleration structure data.
 	 * Resources with this flag can't be bound directly to any shader stage and only can be used with ray tracing APIs.
-	 * This flag is mutually exclusive with all other buffer flags except Static, ReservedResource and ImmediateCommit.
+	 * This flag is mutually exclusive with all other buffer flags except Static and ReservedResource.
 	*/
 	AccelerationStructure   = 1 << 13,
 
@@ -818,9 +818,6 @@ enum class EBufferUsageFlags : uint32
 	*/
 	ReservedResource = 1 << 22,
 
-	/** EXPERIMENTAL: Used with ReservedResource flag to immediately allocate and commit memory on creation. May use N small physical memory allocations instead of a single large one. */
-	ImmediateCommit = 1 << 23,
-
 	// Helper bit-masks
 	AnyDynamic = (Dynamic | Volatile),
 };
@@ -851,7 +848,6 @@ ENUM_CLASS_FLAGS(EBufferUsageFlags);
 #define BUF_NullResource           EBufferUsageFlags::NullResource
 #define BUF_UniformBuffer          EBufferUsageFlags::UniformBuffer
 #define BUF_ReservedResource       EBufferUsageFlags::ReservedResource
-#define BUF_ImmediateCommit        EBufferUsageFlags::ImmediateCommit
 
 enum class EGpuVendorId : uint32
 {
@@ -1066,7 +1062,7 @@ ENUM_CLASS_FLAGS(ETextureCreateFlags);
 #define TexCreate_ReduceMemoryWithTilingMode     ETextureCreateFlags::ReduceMemoryWithTilingMode
 #define TexCreate_Transient                      ETextureCreateFlags::Transient
 #define TexCreate_AtomicCompatible               ETextureCreateFlags::AtomicCompatible
-#define TexCreate_External               		 ETextureCreateFlags::External
+#define TexCreate_External                       ETextureCreateFlags::External
 #define TexCreate_MultiGPUGraphIgnore            ETextureCreateFlags::MultiGPUGraphIgnore
 #define TexCreate_ReservedResource               ETextureCreateFlags::ReservedResource
 #define TexCreate_ImmediateCommit                ETextureCreateFlags::ImmediateCommit
diff --git a/Engine/Source/Runtime/RenderCore/Private/RenderGraphResourcePool.cpp b/Engine/Source/Runtime/RenderCore/Private/RenderGraphResourcePool.cpp
--- a/Engine/Source/Runtime/RenderCore/Private/RenderGraphResourcePool.cpp
+++ b/Engine/Source/Runtime/RenderCore/Private/RenderGraphResourcePool.cpp
@@ -146,7 +146,6 @@ TRefCountPtr<FRDGPooledBuffer> FRDGBufferPool::FindFreeBuffer(FRHICommandListBas
 
 		if (EnumHasAllFlags(Desc.Usage, EBufferUsageFlags::ReservedResource))
 		{
-			checkf(!EnumHasAllFlags(Desc.Usage, EBufferUsageFlags::ImmediateCommit), TEXT("ImmediateCommit is not supported in FRDGBufferPool"));
 			PooledBuffer->CommittedSizeInBytes = 0;
 		}
 
