add rhi unit test for clearing render targets

diff --git a/Engine/Plugins/Tests/RHITests/Source/RHITests/Private/RHITests.spec.cpp b/Engine/Plugins/Tests/RHITests/Source/RHITests/Private/RHITests.spec.cpp
--- a/Engine/Plugins/Tests/RHITests/Source/RHITests/Private/RHITests.spec.cpp
+++ b/Engine/Plugins/Tests/RHITests/Source/RHITests/Private/RHITests.spec.cpp
@@ -51,6 +51,12 @@ void FAutomationRHITest::Define()
 			bool bResult = RunOnRenderThreadSynchronous(FRHITextureTests::Test_RHIClearUAV_Texture3D);
 			TestEqual("Clear Texture3D failed", bResult, 1);
 		});
+
+		It("RHI Clear Render Targets", [this]()
+		{
+			bool bResult = RunOnRenderThreadSynchronous(FRHITextureTests::Test_ClearRenderTargets);
+			TestEqual("Clear Render Targets failed", bResult, 1);
+		});
 	});
 
 	Describe("Test RHI Pixel Format", [this]()
diff --git a/Engine/Plugins/Tests/RHITests/Source/RHITests/Private/RHITestsModule.cpp b/Engine/Plugins/Tests/RHITests/Source/RHITests/Private/RHITestsModule.cpp
--- a/Engine/Plugins/Tests/RHITests/Source/RHITests/Private/RHITestsModule.cpp
+++ b/Engine/Plugins/Tests/RHITests/Source/RHITests/Private/RHITestsModule.cpp
@@ -80,6 +80,13 @@ static bool RunTests_RenderThread(FRHICommandListImmediate& RHICmdList)
 		RUN_TEST(FRHIBufferTests::Test_RHICreateBuffer_Parallel(RHICmdList));
 	}
 
+	// ------------------------------------------------
+	// RT Operations
+	// ------------------------------------------------
+	{
+		RUN_TEST(FRHITextureTests::Test_ClearRenderTargets(RHICmdList));
+	}
+
 	// @todo - add more tests
 	return bResult;
 }
diff --git a/Engine/Plugins/Tests/RHITests/Source/RHITests/Public/RHITextureTests.h b/Engine/Plugins/Tests/RHITests/Source/RHITests/Public/RHITextureTests.h
--- a/Engine/Plugins/Tests/RHITests/Source/RHITests/Public/RHITextureTests.h
+++ b/Engine/Plugins/Tests/RHITests/Source/RHITests/Public/RHITextureTests.h
@@ -7,7 +7,7 @@
 #include "Math/PackedVector.h"
 #include "RHIGPUReadback.h"
 
-//PRAGMA_DISABLE_OPTIMIZATION
+//UE_DISABLE_OPTIMIZATION
 
 class FRHITextureTests
 {
@@ -1522,6 +1522,125 @@ public:
 		UE_LOG(LogRHIUnitTestCommandlet, Display, TEXT("Test passed. \"%s\""), *TestName);
 		return true;
 	}
+
+	template<class TBinaryValue>
+	static bool RunTest_ClearRenderTarget(FRHICommandListImmediate& RHICmdList, EPixelFormat PixelFormat, TArrayView<FVector4f> ClearColors, TArrayView<TBinaryValue> TestValues, uint32 NumChannels, TBinaryValue Tolerance)
+	{
+		const uint32 BytesPerPixel = sizeof(TBinaryValue) * NumChannels;
+		bool bResult = true;
+
+		check(TestValues.Num() == ClearColors.Num() * NumChannels);
+
+		uint32 TestValuesIndex = 0;
+
+		for (int32 ClearColorIndex = 0; ClearColorIndex < ClearColors.Num(); ++ClearColorIndex)
+		{
+			FString ClearTestName = FString::Printf(TEXT("Test_ClearRenderTargets (%s) (%.5f %.5f %.5f %.5f)"), GPixelFormats[PixelFormat].Name, 
+				ClearColors[ClearColorIndex].X, ClearColors[ClearColorIndex].Y, ClearColors[ClearColorIndex].Z, ClearColors[ClearColorIndex].W);
+
+			FRHITextureCreateDesc Desc = FRHITextureCreateDesc(*ClearTestName, ETextureDimension::Texture2D)
+				.SetFormat(PixelFormat)
+				.SetExtent(32, 32)
+				.SetFlags(ETextureCreateFlags::ShaderResource | ETextureCreateFlags::RenderTargetable)
+				.SetClearValue(FClearValueBinding(ClearColors[ClearColorIndex]))
+				.SetInitialState(ERHIAccess::RTV);
+
+			FTextureRHIRef Texture = RHICreateTexture(Desc);
+			check(BytesPerPixel == GPixelFormats[Texture->GetFormat()].BlockBytes);
+
+			FRHIRenderPassInfo RHIRenderPassInfo(Texture, ERenderTargetActions::Clear_Store);
+			RHICmdList.BeginRenderPass(RHIRenderPassInfo, TEXT("ClearRT"));
+			RHICmdList.EndRenderPass();
+			RHICmdList.Transition(FRHITransitionInfo(Texture, ERHIAccess::RTV, ERHIAccess::CopySrc));
+
+			TArrayView<TBinaryValue> TestPixel = TestValues.Slice(TestValuesIndex, NumChannels);
+
+			bool bCurrentColorResult = FRHITextureTests::VerifyTextureContents(*ClearTestName, RHICmdList, Texture,
+				[&](void* Ptr, uint32 MipWidth, uint32 MipHeight, uint32 Width, uint32 Height, uint32 CurrentMipIndex, uint32 CurrentSliceIndex)
+				{
+					for (uint32 Y = 0; Y < MipHeight; ++Y)
+					{
+						uint8* Row = ((uint8*)Ptr) + (Y * Width * BytesPerPixel);
+						for (uint32 X = 0; X < MipWidth; ++X)
+						{
+							TBinaryValue* Pixel = (TBinaryValue*)(Row + X * BytesPerPixel);
+							for (uint32 ChannelIndex=0; ChannelIndex < NumChannels; ++ChannelIndex)
+							{
+								int64 Difference = (int64)TestPixel[ChannelIndex] - (int64)Pixel[ChannelIndex];
+								if (FMath::Abs(Difference) > Tolerance)
+								{
+									return false;
+								}
+							}
+						}
+					}
+					return true;
+				}
+
+			);
+
+			bResult = bResult && bCurrentColorResult;
+			TestValuesIndex += NumChannels;
+		}
+
+		return bResult;
+	}
+
+	static bool Test_ClearRenderTargets(FRHICommandListImmediate& RHICmdList)
+	{
+		bool bResult = true;
+
+		// Float       32-bit     16-bit
+		// 0.2345  = 0x3e7020c5 | 0x3381
+		// 0.8499  = 0x3f59930c | 0x3acc
+		// 0.00145 = 0x3abe0ded | 0x15f0
+		// 0.417   = 0x3ed58106 | 0x36ac
+
+		TArray<FVector4f> ClearColorsLinear =
+		{
+			FVector4f(0.0f, 0.0f, 0.0f, 0.0f) ,
+			FVector4f(0.0f, 0.0f, 0.0f, 1.0f) ,
+			FVector4f(1.0f, 1.0f, 1.0f, 0.0f) ,
+			FVector4f(1.0f, 1.0f, 1.0f, 1.0f) ,
+			FVector4f(0.2345f, 0.8499f, 0.00145f, 0.417f)
+		};
+
+		TArray<uint32> ClearColors32Bits = 
+		{
+			0x00000000, 0x00000000, 0x00000000, 0x00000000,
+			0x00000000, 0x00000000, 0x00000000, 0x3f800000,
+			0x3f800000, 0x3f800000, 0x3f800000, 0x00000000,
+			0x3f800000, 0x3f800000, 0x3f800000, 0x3f800000,
+			0x3e7020c5, 0x3f59930c, 0x3abe0ded, 0x3ed58106,
+		};
+
+		TArray<uint16> ClearColors16Bits = 
+		{
+			0x0000, 0x0000, 0x0000, 0x0000,
+			0x0000, 0x0000, 0x0000, 0x3c00,
+			0x3c00, 0x3c00, 0x3c00, 0x0000,
+			0x3c00, 0x3c00, 0x3c00, 0x3c00,
+			0x3381, 0x3acc, 0x15f0, 0x36ac,
+		};
+
+		TArray<uint8> ClearColors8Bits = 
+		{
+			0x00, 0x00, 0x00, 0x00,
+			0x00, 0x00, 0x00, 0xFF,
+			0xFF, 0xFF, 0xFF, 0x00,
+			0xFF, 0xFF, 0xFF, 0xFF,
+			0x3c, 0xd9, 0x00, 0x6a,
+		};
+
+		RUN_TEST(RunTest_ClearRenderTarget<uint32>(RHICmdList, PF_A32B32G32R32F, ClearColorsLinear, ClearColors32Bits, 4, 0));
+		// Some platforms do not behave consistently wrt rounding during float -> float16 or 8bits unorm conversion. Use a tolerance until we fix these
+		RUN_TEST(RunTest_ClearRenderTarget<uint16>(RHICmdList, PF_FloatRGBA, ClearColorsLinear, ClearColors16Bits, 4, 1));
+		RUN_TEST(RunTest_ClearRenderTarget<uint8>(RHICmdList, PF_R8G8B8A8, ClearColorsLinear, ClearColors8Bits, 4, 1));
+
+		RHICmdList.ImmediateFlush(EImmediateFlushType::FlushRHIThreadFlushResources);
+		return bResult;
+
+	}
 };
 
-//PRAGMA_ENABLE_OPTIMIZATION
+//UE_ENABLE_OPTIMIZATION
