Extend RHI unit tests to cover basic direct and indirect drawing with instance and vertex offsets

diff --git a/Engine/Plugins/Tests/RHITests/Shaders/Private/TestDrawInstanced.usf b/Engine/Plugins/Tests/RHITests/Shaders/Private/TestDrawInstanced.usf
--- a/Engine/Plugins/Tests/RHITests/Shaders/Private/TestDrawInstanced.usf
+++ b/Engine/Plugins/Tests/RHITests/Shaders/Private/TestDrawInstanced.usf
@@ -6,33 +6,19 @@
 // Used to validate that instanced draws with base instance offset work correctly.
 
 void TestDrawInstancedMainVS(
-	in uint CustomInstanceID  : ATTRIBUTE0,
-	in uint VertexID		: SV_VertexID,
-	out uint OutInstanceID  : Texcoord0,
-	out float4 OutPosition	: SV_Position
+	in float4 VertexPosition : ATTRIBUTE0,
+	in uint CustomInstanceID : ATTRIBUTE1,
+	out uint OutInstanceID : Texcoord0,
+	out float4 OutPosition : SV_Position
 )
 {
+	OutPosition = VertexPosition;
 	OutInstanceID = CustomInstanceID;
-
-#if 0
-	// Produces broken code on some platforms
-	switch(VertexID)
-	{
-		default:
-		case 0:	OutPosition = float4(-1, -1, 0, 1); break;
-		case 1:	OutPosition = float4(-1, +3, 0, 1); break;
-		case 2:	OutPosition = float4(+3, -1, 0, 1); break;
-	}
-#else
-	if      (VertexID == 0) { OutPosition = float4(-1, -1, 0, 1); }
-	else if (VertexID == 1) { OutPosition = float4(-1, +3, 0, 1); }
-	else                    { OutPosition = float4(+3, -1, 0, 1); }
-#endif
 }
 
 RWBuffer<uint>	OutDrawnInstances;
 void TestDrawInstancedMainPS(
-	in nointerpolation uint InstanceID  : Texcoord0,
+	in nointerpolation uint InstanceID : Texcoord0,
 	out float4 OutColor	: SV_Target0
 )
 {
diff --git a/Engine/Plugins/Tests/RHITests/Source/RHITests/Private/RHIDrawTests.cpp b/Engine/Plugins/Tests/RHITests/Source/RHITests/Private/RHIDrawTests.cpp
--- a/Engine/Plugins/Tests/RHITests/Source/RHITests/Private/RHIDrawTests.cpp
+++ b/Engine/Plugins/Tests/RHITests/Source/RHITests/Private/RHIDrawTests.cpp
@@ -5,6 +5,7 @@
 #include "CommonRenderResources.h"
 #include "RenderCaptureInterface.h"
 #include "RHIStaticStates.h"
+#include "ShaderCompilerCore.h"
 
 class FTestDrawInstancedVS : public FGlobalShader
 {
@@ -16,6 +17,7 @@ class FTestDrawInstancedVS : public FGlobalShader
 	}
 	static void ModifyCompilationEnvironment(const FGlobalShaderPermutationParameters& Parameters, FShaderCompilerEnvironment& OutEnvironment)
 	{
+		OutEnvironment.CompilerFlags.Add(CFLAG_IndirectDraw);
 	}
 };
 IMPLEMENT_GLOBAL_SHADER(FTestDrawInstancedVS, "/Plugin/RHITests/Private/TestDrawInstanced.usf", "TestDrawInstancedMainVS", SF_Vertex);
@@ -36,6 +38,9 @@ class FTestDrawInstancedPS : public FGlobalShader
 };
 IMPLEMENT_GLOBAL_SHADER(FTestDrawInstancedPS, "/Plugin/RHITests/Private/TestDrawInstanced.usf", "TestDrawInstancedMainPS", SF_Pixel);
 
+namespace
+{
+
 template <typename T>
 static FBufferRHIRef CreateBufferWithData(EBufferUsageFlags UsageFlags, ERHIAccess ResourceState, const TCHAR* Name, TConstArrayView<T> Data)
 {
@@ -49,33 +54,215 @@ static FBufferRHIRef CreateBufferWithData(EBufferUsageFlags UsageFlags, ERHIAcce
 	return Buffer;
 }
 
-bool FRHIDrawTests::Test_MultiDrawIndirect(FRHICommandListImmediate& RHICmdList)
+template<typename T>
+static FBufferRHIRef CreateBufferWithData(EBufferUsageFlags UsageFlags, ERHIAccess ResourceState, const TCHAR* Name, TArrayView<T> Data)
 {
-	if (!GRHIGlobals.SupportsMultiDrawIndirect)
+	return CreateBufferWithData(UsageFlags, ResourceState, Name, TConstArrayView<T>(Data.GetData(), Data.Num()));
+}
+
+// Structure to initialize common resources required for various draw tests
+struct FDrawTestResources
+{
+	FDrawTestResources(FRHICommandListImmediate& RHICmdList)
+		: VertexShader(GetGlobalShaderMap(GMaxRHIFeatureLevel))
+		, PixelShader(GetGlobalShaderMap(GMaxRHIFeatureLevel))
+	{		
+		FRHITextureDesc RenderTargetTextureDesc(ETextureDimension::Texture2D, ETextureCreateFlags::RenderTargetable, PF_B8G8R8A8, FClearValueBinding(), RenderTargetSize, 1, 1, 1, 1, 0);
+		FRHITextureCreateDesc RenderTargetCreateDesc(RenderTargetTextureDesc, ERHIAccess::RTV, TEXT("DrawTest_RenderTarget"));
+		RenderTarget = RHICreateTexture(RenderTargetCreateDesc);
+
+		FVertexDeclarationElementList VertexDeclarationElements;
+		// Position vertex element
+		VertexDeclarationElements.Add(FVertexElement(0, 0, VET_Float4, 0, 16, false /*per instance frequency*/));
+		// Vertex element to access instance ID buffer (uint per instance)
+		VertexDeclarationElements.Add(FVertexElement(1, 0, VET_UInt, 1, 4, true /*per instance frequency*/));
+
+		VertexDeclarationRHI = PipelineStateCache::GetOrCreateVertexDeclaration(VertexDeclarationElements);
+
+		GraphicsPSOInit.BoundShaderState.VertexShaderRHI = VertexShader.GetVertexShader();
+		GraphicsPSOInit.BoundShaderState.VertexDeclarationRHI = VertexDeclarationRHI;
+		GraphicsPSOInit.BoundShaderState.PixelShaderRHI = PixelShader.GetPixelShader();
+		GraphicsPSOInit.DepthStencilState = TStaticDepthStencilState<false, CF_Always>::GetRHI();
+		GraphicsPSOInit.BlendState = TStaticBlendState<>::GetRHI();
+		GraphicsPSOInit.RasterizerState = TStaticRasterizerState<>::GetRHI();
+		GraphicsPSOInit.PrimitiveType = EPrimitiveType::PT_TriangleList;
+
+		// D3D12 does not have a way to get the base instance ID (SV_InstanceID always starts from 0), so we must emulate it...
+		uint32 InstanceIDs[MaxInstances];
+		for (uint32 i = 0; i < MaxInstances; ++i)
+		{
+			InstanceIDs[i] = i;
+		}
+		InstanceIDBuffer = CreateBufferWithData(EBufferUsageFlags::VertexBuffer, ERHIAccess::VertexOrIndexBuffer, TEXT("DrawTest_InstanceID"), MakeArrayView(InstanceIDs));
+
+		const uint16 Indices[3] = { 0, 1, 2 };
+		IndexBuffer = CreateBufferWithData(EBufferUsageFlags::IndexBuffer, ERHIAccess::VertexOrIndexBuffer, TEXT("DrawTest_IndexBuffer"), MakeArrayView(Indices));
+
+		TArray<FVector4f> Vertices;
+		Vertices.Reserve(NumTotalVertices);
+		for (uint32 i = 0; i < NumValidVertices; ++i)
+		{
+			switch (i % 3)
+			{
+			case 0:	Vertices.Add(FVector4f(-1.0f, -1.0f, 0.0f, 1.0f)); break;
+			case 1:	Vertices.Add(FVector4f(-1.0f, +3.0f, 0.0f, 1.0f)); break;
+			case 2:	Vertices.Add(FVector4f(+3.0f, -1.0f, 0.0f, 1.0f)); break;
+			}
+		}
+		for (uint32 i = NumValidVertices; i < NumTotalVertices; ++i)
+		{
+			Vertices.Add(FVector4f(0.0f, 0.0f, 0.0f, 1.0f));
+		}
+		VertexBuffer = CreateBufferWithData(EBufferUsageFlags::VertexBuffer, ERHIAccess::VertexOrIndexBuffer, TEXT("DrawTest_VertexBuffer"), MakeArrayView(Vertices));
+
+		static constexpr uint32 OutputBufferStride = sizeof(uint32);
+		static constexpr uint32 OutputBufferSize = OutputBufferStride * MaxInstances;
+		FRHIResourceCreateInfo OutputBufferCreateInfo(TEXT("DrawTest_OutputBuffer"));
+		OutputBuffer = RHICmdList.CreateBuffer(OutputBufferSize, EBufferUsageFlags::UnorderedAccess | EBufferUsageFlags::SourceCopy, OutputBufferStride, ERHIAccess::UAVCompute, OutputBufferCreateInfo);
+
+		OutputBufferUAV = RHICmdList.CreateUnorderedAccessView(OutputBuffer,
+			FRHIViewDesc::CreateBufferUAV()
+			.SetType(FRHIViewDesc::EBufferType::Typed)
+			.SetFormat(PF_R32_UINT));
+	}
+
+	static constexpr uint32 MaxInstances = 8;
+
+	FVertexDeclarationRHIRef VertexDeclarationRHI;
+	TShaderMapRef<FTestDrawInstancedVS> VertexShader;
+	TShaderMapRef<FTestDrawInstancedPS> PixelShader;
+	FGraphicsPipelineStateInitializer GraphicsPSOInit;
+
+	FIntPoint RenderTargetSize = FIntPoint(4, 4);
+	FTextureRHIRef RenderTarget;
+
+	FBufferRHIRef InstanceIDBuffer;
+	FBufferRHIRef IndexBuffer;
+
+	// 2 full screen triangles, followed by a degenerate triangle
+	static constexpr uint32 NumTotalVertices = 9;
+	static constexpr uint32 NumValidVertices = 6;
+	FBufferRHIRef VertexBuffer;
+
+	FBufferRHIRef OutputBuffer;
+
+	FUnorderedAccessViewRHIRef OutputBufferUAV;
+};
+
+}
+
+bool FRHIDrawTests::InternalDrawBaseVertexAndInstance(FRHICommandListImmediate& RHICmdList, EDrawKind DrawKind, const TCHAR* TestName)
+{
+	FDrawTestResources Resources(RHICmdList);
+
+	RHICmdList.ClearUAVUint(Resources.OutputBufferUAV, FUintVector4(0));
+
+	FRHITexture* ColorRTs[1] = { Resources.RenderTarget.GetReference() };
+	FRHIRenderPassInfo RenderPassInfo(1, ColorRTs, ERenderTargetActions::DontLoad_DontStore);
+
+	const uint32 FirstInvalidVertex = Resources.NumValidVertices;
+
+	const FRHIDrawIndexedIndirectParameters DrawArgs[4] =
+	{
+		// IndexCountPerInstance, InstanceCount, StartIndexLocation, BaseVertexLocation, StartInstanceLocation
+		{3,                       1,             0,                  0,                  0}, // InstanceID = 0, expected to be drawn. No vertex of instance offset
+		{3,                       1,             0,                  0,                  1}, // InstanceID = 1, expected to be drawn
+		{3,                       1,             0,                  FirstInvalidVertex, 2}, // InstanceID = 2, expected to be culled by rendering a degenerate triangle (see vertex buffer initialization)
+		{3,                       1,             0,                  0,                  3}, // InstanceID = 3, expected to be drawn
+	};
+
+	const uint32 ExpectedDrawnInstances[Resources.MaxInstances] = { 1, 1, 0, 1, 0, 0, 0, 0 };
+
+	FBufferRHIRef DrawArgBuffer;
+	if (DrawKind == EDrawKind::Indirect)
+	{
+		DrawArgBuffer = CreateBufferWithData(EBufferUsageFlags::DrawIndirect | EBufferUsageFlags::UnorderedAccess | EBufferUsageFlags::VertexBuffer, ERHIAccess::IndirectArgs,
+			TEXT("InternalDrawBaseVertexAndInstance_DrawArgs"), MakeArrayView(DrawArgs));
+	}
+
+	RHICmdList.Transition(FRHITransitionInfo(Resources.OutputBufferUAV, ERHIAccess::UAVCompute, ERHIAccess::UAVGraphics, EResourceTransitionFlags::None));
+	RHICmdList.BeginUAVOverlap(); // Output UAV can be written without syncs between draws (each draw is expected to write into different slots)
+
+	RHICmdList.BeginRenderPass(RenderPassInfo, TestName);
+	RHICmdList.SetViewport(0, 0, 0, float(Resources.RenderTargetSize.X), float(Resources.RenderTargetSize.Y), 1);
+
+	RHICmdList.ApplyCachedRenderTargets(Resources.GraphicsPSOInit);
+	SetGraphicsPipelineState(RHICmdList, Resources.GraphicsPSOInit, 0);
+
+	check(Resources.InstanceIDBuffer->GetStride() == 4);
+	RHICmdList.SetStreamSource(0, Resources.VertexBuffer, 0);
+	RHICmdList.SetStreamSource(1, Resources.InstanceIDBuffer, 0);
+
+	FRHIBatchedShaderParameters ShaderParameters;
+	ShaderParameters.SetUAVParameter(Resources.PixelShader->OutDrawnInstances.GetBaseIndex(), Resources.OutputBufferUAV);
+	RHICmdList.SetBatchedShaderParameters(Resources.PixelShader.GetPixelShader(), ShaderParameters);
+
+	if (DrawKind == EDrawKind::Direct)
+	{
+		for (const FRHIDrawIndexedIndirectParameters& Draw : DrawArgs)
+		{
+			RHICmdList.DrawIndexedPrimitive(Resources.IndexBuffer,
+				Draw.BaseVertexLocation,
+				Draw.StartInstanceLocation,
+				Draw.IndexCountPerInstance,
+				Draw.StartIndexLocation,
+				Draw.IndexCountPerInstance / 3 /*NumPrimitives*/,
+				Draw.InstanceCount);
+		}
+	}
+	else
+	{
+		for (const FRHIDrawIndexedIndirectParameters& Draw : DrawArgs)
+		{
+			const uint32 ArgumentOffset = uint32(uint64(&Draw) - uint64(DrawArgs));
+			RHICmdList.DrawIndexedPrimitiveIndirect(Resources.IndexBuffer, DrawArgBuffer, ArgumentOffset);
+		}
+	}
+
+	RHICmdList.EndRenderPass();
+
+	RHICmdList.EndUAVOverlap();
+
+	RHICmdList.Transition(FRHITransitionInfo(Resources.OutputBufferUAV, ERHIAccess::UAVGraphics, ERHIAccess::CopySrc, EResourceTransitionFlags::None));
+
+	TConstArrayView<uint8> ExpectedOutputView = MakeArrayView(reinterpret_cast<const uint8*>(ExpectedDrawnInstances), sizeof(ExpectedDrawnInstances));
+	bool bSucceeded = FRHIBufferTests::VerifyBufferContents(TestName, RHICmdList, Resources.OutputBuffer, ExpectedOutputView);
+
+	return bSucceeded;
+}
+
+bool FRHIDrawTests::Test_DrawBaseVertexAndInstanceDirect(FRHICommandListImmediate& RHICmdList)
+{
+	if (!GRHIGlobals.SupportsBaseVertexIndex || !GRHIGlobals.SupportsFirstInstance)
 	{
 		return true;
 	}
 
-	// Probably could/should automatically enable in the outer scope when running RHI Unit Tests
-	// RenderCaptureInterface::FScopedCapture RenderCapture(true /*bEnable*/, &RHICmdList, TEXT("Test_MultiDrawIndirect"));
+	return InternalDrawBaseVertexAndInstance(RHICmdList, EDrawKind::Direct, TEXT("Test_DrawBaseVertexAndInstanceDirect"));
+}
 
-	static constexpr uint32 MaxInstances = 8;
+bool FRHIDrawTests::Test_DrawBaseVertexAndInstanceIndirect(FRHICommandListImmediate& RHICmdList)
+{
+	if (!GRHIGlobals.SupportsBaseVertexIndex || !GRHIGlobals.SupportsFirstInstance)
+	{
+		return true;
+	}
+
+	return InternalDrawBaseVertexAndInstance(RHICmdList, EDrawKind::Indirect, TEXT("Test_DrawBaseVertexAndInstanceIndirect"));
+}
 
-	// D3D12 does not have a way to get the base instance ID (SV_InstanceID always starts from 0), so we must emulate it...
-	const uint32 InstanceIDs[MaxInstances] = { 0, 1, 2, 3, 4, 5, 6, 7 };
-	FBufferRHIRef InstanceIDBuffer = CreateBufferWithData(EBufferUsageFlags::VertexBuffer, ERHIAccess::VertexOrIndexBuffer, TEXT("Test_MultiDrawIndirect_InstanceID"), MakeArrayView(InstanceIDs));
 
-	FVertexDeclarationElementList VertexDeclarationElements;
-	VertexDeclarationElements.Add(FVertexElement(0, 0, VET_UInt, 0, 4, true /*per instance frequency*/));
-	FVertexDeclarationRHIRef VertexDeclarationRHI = PipelineStateCache::GetOrCreateVertexDeclaration(VertexDeclarationElements);
+bool FRHIDrawTests::Test_MultiDrawIndirect(FRHICommandListImmediate& RHICmdList)
+{
+	if (!GRHIGlobals.SupportsMultiDrawIndirect)
+	{
+		return true;
+	}
 
-	const uint16 Indices[3] = { 0, 1, 2 };
-	FBufferRHIRef IndexBuffer = CreateBufferWithData(EBufferUsageFlags::IndexBuffer, ERHIAccess::VertexOrIndexBuffer, TEXT("Test_MultiDrawIndirect_IndexBuffer"), MakeArrayView(Indices));
+	// Probably could/should automatically enable in the outer scope when running RHI Unit Tests
+	// RenderCaptureInterface::FScopedCapture RenderCapture(true /*bEnable*/, &RHICmdList, TEXT("Test_MultiDrawIndirect"));
 
-	static constexpr uint32 OutputBufferStride = sizeof(uint32);
-	static constexpr uint32 OutputBufferSize = OutputBufferStride * MaxInstances;
-	FRHIResourceCreateInfo OutputBufferCreateInfo(TEXT("Test_MultiDrawIndirect_OutputBuffer"));
-	FBufferRHIRef OutputBuffer = RHICmdList.CreateBuffer(OutputBufferSize, EBufferUsageFlags::UnorderedAccess | EBufferUsageFlags::SourceCopy, OutputBufferStride, ERHIAccess::UAVCompute, OutputBufferCreateInfo);
+	FDrawTestResources Resources(RHICmdList);
 
 	const uint32 CountValues[4] = { 1, 1, 16, 0 };
 	FBufferRHIRef CountBuffer = CreateBufferWithData(EBufferUsageFlags::DrawIndirect | EBufferUsageFlags::UnorderedAccess, ERHIAccess::IndirectArgs, TEXT("Test_MultiDrawIndirect_Count"), MakeArrayView(CountValues));
@@ -92,83 +279,61 @@ bool FRHIDrawTests::Test_MultiDrawIndirect(FRHICommandListImmediate& RHICmdList)
 		{3, 1, 0, 0, 7}, // this draw is expected to never execute
 	};
 
-	const uint32 ExpectedDrawnInstances[MaxInstances] = { 1, 0, 1, 1, 0, 1, 1, 0 };
+	const uint32 ExpectedDrawnInstances[Resources.MaxInstances] = { 1, 0, 1, 1, 0, 1, 1, 0 };
 
 	FBufferRHIRef DrawArgBuffer = CreateBufferWithData(EBufferUsageFlags::DrawIndirect | EBufferUsageFlags::UnorderedAccess | EBufferUsageFlags::VertexBuffer, ERHIAccess::IndirectArgs,
 		TEXT("Test_MultiDrawIndirect_DrawArgs"), MakeArrayView(DrawArgs));
 
-	FUnorderedAccessViewRHIRef OutputBufferUAV = RHICmdList.CreateUnorderedAccessView(OutputBuffer, 
-		FRHIViewDesc::CreateBufferUAV()
-		.SetType(FRHIViewDesc::EBufferType::Typed)
-		.SetFormat(PF_R32_UINT));
-
-	RHICmdList.ClearUAVUint(OutputBufferUAV, FUintVector4(0));
-
-	const FIntPoint RenderTargetSize(4, 4);
-	FRHITextureDesc RenderTargetTextureDesc(ETextureDimension::Texture2D, ETextureCreateFlags::RenderTargetable, PF_B8G8R8A8, FClearValueBinding(), RenderTargetSize, 1, 1, 1, 1, 0);
-	FRHITextureCreateDesc RenderTargetCreateDesc(RenderTargetTextureDesc, ERHIAccess::RTV, TEXT("Test_MultiDrawIndirect_RenderTarget"));
-	FTextureRHIRef RenderTarget = RHICreateTexture(RenderTargetCreateDesc);
+	RHICmdList.ClearUAVUint(Resources.OutputBufferUAV, FUintVector4(0));
 
-	TShaderMapRef<FTestDrawInstancedVS> VertexShader(GetGlobalShaderMap(GMaxRHIFeatureLevel));
-	TShaderMapRef<FTestDrawInstancedPS> PixelShader(GetGlobalShaderMap(GMaxRHIFeatureLevel));
-
-	FGraphicsPipelineStateInitializer GraphicsPSOInit;
-	
-	GraphicsPSOInit.BoundShaderState.VertexShaderRHI = VertexShader.GetVertexShader();
-	GraphicsPSOInit.BoundShaderState.VertexDeclarationRHI = VertexDeclarationRHI;
-	GraphicsPSOInit.BoundShaderState.PixelShaderRHI = PixelShader.GetPixelShader();
-	GraphicsPSOInit.DepthStencilState = TStaticDepthStencilState<false, CF_Always>::GetRHI();
-	GraphicsPSOInit.BlendState = TStaticBlendState<>::GetRHI();
-	GraphicsPSOInit.RasterizerState = TStaticRasterizerState<>::GetRHI();
-	GraphicsPSOInit.PrimitiveType = EPrimitiveType::PT_TriangleList;
-
-	FRHITexture* ColorRTs[1] = { RenderTarget.GetReference() };
+	FRHITexture* ColorRTs[1] = { Resources.RenderTarget.GetReference() };
 	FRHIRenderPassInfo RenderPassInfo(1, ColorRTs, ERenderTargetActions::DontLoad_DontStore);
 
-	RHICmdList.Transition(FRHITransitionInfo(OutputBufferUAV, ERHIAccess::UAVCompute, ERHIAccess::UAVGraphics, EResourceTransitionFlags::None));
+	RHICmdList.Transition(FRHITransitionInfo(Resources.OutputBufferUAV, ERHIAccess::UAVCompute, ERHIAccess::UAVGraphics, EResourceTransitionFlags::None));
 	RHICmdList.BeginUAVOverlap(); // Output UAV can be written without syncs between draws (each draw is expected to write into different slots)
 
 	RHICmdList.BeginRenderPass(RenderPassInfo, TEXT("Test_MultiDrawIndirect"));
-	RHICmdList.SetViewport(0, 0, 0, float(RenderTargetSize.X), float(RenderTargetSize.Y), 1);
+	RHICmdList.SetViewport(0, 0, 0, float(Resources.RenderTargetSize.X), float(Resources.RenderTargetSize.Y), 1);
 
-	RHICmdList.ApplyCachedRenderTargets(GraphicsPSOInit);
-	SetGraphicsPipelineState(RHICmdList, GraphicsPSOInit, 0);
+	RHICmdList.ApplyCachedRenderTargets(Resources.GraphicsPSOInit);
+	SetGraphicsPipelineState(RHICmdList, Resources.GraphicsPSOInit, 0);
 
-	check(InstanceIDBuffer->GetStride() == 4);
-	RHICmdList.SetStreamSource(0, InstanceIDBuffer, 0);
+	check(Resources.InstanceIDBuffer->GetStride() == 4);
+	RHICmdList.SetStreamSource(0, Resources.VertexBuffer, 0);
+	RHICmdList.SetStreamSource(1, Resources.InstanceIDBuffer, 0);
 
 	FRHIBatchedShaderParameters ShaderParameters;
-	ShaderParameters.SetUAVParameter(PixelShader->OutDrawnInstances.GetBaseIndex(), OutputBufferUAV);
-	RHICmdList.SetBatchedShaderParameters(PixelShader.GetPixelShader(), ShaderParameters);
+	ShaderParameters.SetUAVParameter(Resources.PixelShader->OutDrawnInstances.GetBaseIndex(), Resources.OutputBufferUAV);
+	RHICmdList.SetBatchedShaderParameters(Resources.PixelShader.GetPixelShader(), ShaderParameters);
 
 	const uint32 DrawArgsStride = sizeof(DrawArgs[0]);
 	const uint32 CountStride = sizeof(CountValues[0]);
 
-	RHICmdList.MultiDrawIndexedPrimitiveIndirect(IndexBuffer, 
+	RHICmdList.MultiDrawIndexedPrimitiveIndirect(Resources.IndexBuffer, 
 		DrawArgBuffer, DrawArgsStride*0, // 1 sub-draw with instance index 0
 		CountBuffer, CountStride*0, // count buffer contains 1 in this slot
 		5 // expect to draw only 1 instance due to GPU-side upper bound
 	);
 
-	RHICmdList.MultiDrawIndexedPrimitiveIndirect(IndexBuffer,
+	RHICmdList.MultiDrawIndexedPrimitiveIndirect(Resources.IndexBuffer,
 		DrawArgBuffer, DrawArgsStride*1, // 1 sub-draw with 2 instances at base index 2
 		CountBuffer, CountStride*1, // count buffer contains 1 in this slot
 		4 // expect to draw only 1 instance due to GPU-side upper bound
 	);
 
-	RHICmdList.MultiDrawIndexedPrimitiveIndirect(IndexBuffer,
+	RHICmdList.MultiDrawIndexedPrimitiveIndirect(Resources.IndexBuffer,
 		DrawArgBuffer, DrawArgsStride*2, // 2 sub-draws with 1 instance each starting at base index 5
 		CountBuffer, CountStride*2, // count buffer contains 16 in this slot
 		2 // expect to draw only 2 instances due to CPU-side upper bound
 	);
 
-	RHICmdList.MultiDrawIndexedPrimitiveIndirect(IndexBuffer,
+	RHICmdList.MultiDrawIndexedPrimitiveIndirect(Resources.IndexBuffer,
 		DrawArgBuffer, DrawArgsStride*4, // 1 sub-draw with 1 instance each starting at base index 7
 		CountBuffer, CountStride*3, // count buffer contains 0 in this slot
 		1 // expect to skip the draw due to GPU-side count of 0
 	);
 
-	RHICmdList.MultiDrawIndexedPrimitiveIndirect(IndexBuffer,
+	RHICmdList.MultiDrawIndexedPrimitiveIndirect(Resources.IndexBuffer,
 		DrawArgBuffer, DrawArgsStride*4, // 1 sub-draw with 1 instance each starting at base index 7
 		CountBuffer, CountStride*0, // count buffer contains 1 in this slot
 		0 // expect to skip the draw due to CPU-side count of 0
@@ -178,10 +343,10 @@ bool FRHIDrawTests::Test_MultiDrawIndirect(FRHICommandListImmediate& RHICmdList)
 
 	RHICmdList.EndUAVOverlap();
 
-	RHICmdList.Transition(FRHITransitionInfo(OutputBufferUAV, ERHIAccess::UAVGraphics, ERHIAccess::CopySrc, EResourceTransitionFlags::None));
+	RHICmdList.Transition(FRHITransitionInfo(Resources.OutputBufferUAV, ERHIAccess::UAVGraphics, ERHIAccess::CopySrc, EResourceTransitionFlags::None));
 
 	TConstArrayView<uint8> ExpectedOutputView = MakeArrayView(reinterpret_cast<const uint8*>(ExpectedDrawnInstances), sizeof(ExpectedDrawnInstances));
-	bool bSucceeded = FRHIBufferTests::VerifyBufferContents(TEXT("Test_MultiDrawIndirect"), RHICmdList, OutputBuffer, ExpectedOutputView);
+	bool bSucceeded = FRHIBufferTests::VerifyBufferContents(TEXT("Test_MultiDrawIndirect"), RHICmdList, Resources.OutputBuffer, ExpectedOutputView);
 
 	return bSucceeded;
 }
diff --git a/Engine/Plugins/Tests/RHITests/Source/RHITests/Private/RHITests.spec.cpp b/Engine/Plugins/Tests/RHITests/Source/RHITests/Private/RHITests.spec.cpp
--- a/Engine/Plugins/Tests/RHITests/Source/RHITests/Private/RHITests.spec.cpp
+++ b/Engine/Plugins/Tests/RHITests/Source/RHITests/Private/RHITests.spec.cpp
@@ -118,6 +118,16 @@ void FAutomationRHITest::Define()
 
 	Describe("Test RHI Draw", [this]()
 	{
+		It("RHI DrawBaseVertexAndInstanceDirect", [this]()
+		{
+			bool bResult = RunOnRenderThreadSynchronous(FRHIDrawTests::Test_DrawBaseVertexAndInstanceDirect);
+			TestEqual("RHI DrawBaseVertexAndInstanceDirect", bResult, 1);
+		});
+		It("RHI DrawBaseVertexAndInstanceIndirect", [this]()
+		{
+			bool bResult = RunOnRenderThreadSynchronous(FRHIDrawTests::Test_DrawBaseVertexAndInstanceIndirect);
+			TestEqual("RHI DrawBaseVertexAndInstanceIndirect", bResult, 1);
+		});
 		It("RHI MultiDrawIndirect", [this]()
 		{
 			bool bResult = RunOnRenderThreadSynchronous(FRHIDrawTests::Test_MultiDrawIndirect);
diff --git a/Engine/Plugins/Tests/RHITests/Source/RHITests/Private/RHITestsModule.cpp b/Engine/Plugins/Tests/RHITests/Source/RHITests/Private/RHITestsModule.cpp
--- a/Engine/Plugins/Tests/RHITests/Source/RHITests/Private/RHITestsModule.cpp
+++ b/Engine/Plugins/Tests/RHITests/Source/RHITests/Private/RHITestsModule.cpp
@@ -31,6 +31,8 @@ static bool RunTests_RenderThread(FRHICommandListImmediate& RHICmdList)
 	// Drawing
 	// ------------------------------------------------
 	{
+		RUN_TEST(FRHIDrawTests::Test_DrawBaseVertexAndInstanceDirect(RHICmdList));
+		RUN_TEST(FRHIDrawTests::Test_DrawBaseVertexAndInstanceIndirect(RHICmdList));
 		RUN_TEST(FRHIDrawTests::Test_MultiDrawIndirect(RHICmdList));
 	}
 
diff --git a/Engine/Plugins/Tests/RHITests/Source/RHITests/Public/RHIDrawTests.h b/Engine/Plugins/Tests/RHITests/Source/RHITests/Public/RHIDrawTests.h
--- a/Engine/Plugins/Tests/RHITests/Source/RHITests/Public/RHIDrawTests.h
+++ b/Engine/Plugins/Tests/RHITests/Source/RHITests/Public/RHIDrawTests.h
@@ -8,6 +8,18 @@ class FRHIDrawTests
 {
 public:
 
+	static bool Test_DrawBaseVertexAndInstanceDirect(FRHICommandListImmediate& RHICmdList);
+	static bool Test_DrawBaseVertexAndInstanceIndirect(FRHICommandListImmediate& RHICmdList);
 	static bool Test_MultiDrawIndirect(FRHICommandListImmediate& RHICmdList);
+
+private:
+
+	enum class EDrawKind
+	{
+		Direct,
+		Indirect,
+	};
+
+	static bool InternalDrawBaseVertexAndInstance(FRHICommandListImmediate& RHICmdList, EDrawKind DrawKind, const TCHAR* TestName);
 };
 
