Extend RHI basic draw unit tests to cover more edge cases

* Test non-zero index buffer offset
* Test combined index buffer offset and base vertex

diff --git a/Engine/Plugins/Tests/RHITests/Source/RHITests/Private/RHIDrawTests.cpp b/Engine/Plugins/Tests/RHITests/Source/RHITests/Private/RHIDrawTests.cpp
--- a/Engine/Plugins/Tests/RHITests/Source/RHITests/Private/RHIDrawTests.cpp
+++ b/Engine/Plugins/Tests/RHITests/Source/RHITests/Private/RHIDrawTests.cpp
@@ -95,7 +95,15 @@ struct FDrawTestResources
 		}
 		InstanceIDBuffer = CreateBufferWithData(EBufferUsageFlags::VertexBuffer, ERHIAccess::VertexOrIndexBuffer, TEXT("DrawTest_InstanceID"), MakeArrayView(InstanceIDs));
 
-		const uint16 Indices[3] = { 0, 1, 2 };
+		// Indices for 4 triangles
+		const uint16 Indices[NumTotalVertices] =
+		{
+			0, 1, 2, // valid
+			3, 4, 5, // valid
+			6, 7, 8, // valid
+			9, 10, 11 // degenerate
+		};
+
 		IndexBuffer = CreateBufferWithData(EBufferUsageFlags::IndexBuffer, ERHIAccess::VertexOrIndexBuffer, TEXT("DrawTest_IndexBuffer"), MakeArrayView(Indices));
 
 		TArray<FVector4f> Vertices;
@@ -139,9 +147,9 @@ struct FDrawTestResources
 	FBufferRHIRef InstanceIDBuffer;
 	FBufferRHIRef IndexBuffer;
 
-	// 2 full screen triangles, followed by a degenerate triangle
-	static constexpr uint32 NumTotalVertices = 9;
-	static constexpr uint32 NumValidVertices = 6;
+	// 3 full screen triangles, followed by a degenerate triangle
+	static constexpr uint32 NumTotalVertices = 12;
+	static constexpr uint32 NumValidVertices = 9;
 	FBufferRHIRef VertexBuffer;
 
 	FBufferRHIRef OutputBuffer;
@@ -162,16 +170,18 @@ bool FRHIDrawTests::InternalDrawBaseVertexAndInstance(FRHICommandListImmediate&
 
 	const uint32 FirstInvalidVertex = Resources.NumValidVertices;
 
-	const FRHIDrawIndexedIndirectParameters DrawArgs[4] =
+	const FRHIDrawIndexedIndirectParameters DrawArgs[6] =
 	{
 		// IndexCountPerInstance, InstanceCount, StartIndexLocation, BaseVertexLocation, StartInstanceLocation
 		{3,                       1,             0,                  0,                  0}, // InstanceID = 0, expected to be drawn. No vertex of instance offset
 		{3,                       1,             0,                  0,                  1}, // InstanceID = 1, expected to be drawn
 		{3,                       1,             0,                  FirstInvalidVertex, 2}, // InstanceID = 2, expected to be culled by rendering a degenerate triangle (see vertex buffer initialization)
-		{3,                       1,             0,                  0,                  3}, // InstanceID = 3, expected to be drawn
+		{3,                       1,             3,                  0,                  3}, // InstanceID = 3, expected to be drawn (vertices 3,4,5 via index buffer offset)
+		{3,                       1,             0,                  3,                  4}, // InstanceID = 4, expected to be drawn (vertices 3,4,5 via base vertex)
+		{3,                       1,             3,                  3,                  5}, // InstanceID = 5, expected to be drawn (vertices 6,7,8 via base vertex and index buffer offset)
 	};
 
-	const uint32 ExpectedDrawnInstances[Resources.MaxInstances] = { 1, 1, 0, 1, 0, 0, 0, 0 };
+	const uint32 ExpectedDrawnInstances[Resources.MaxInstances] = { 1, 1, 0, 1, 1, 1, 0, 0 };
 
 	FBufferRHIRef DrawArgBuffer;
 	if (DrawKind == EDrawKind::Indirect)
