Implement basic support for reserved volume textures in D3D12 RHI

diff --git a/Engine/Plugins/Tests/RHITests/Source/RHITests/Private/RHIReservedResourceTests.cpp b/Engine/Plugins/Tests/RHITests/Source/RHITests/Private/RHIReservedResourceTests.cpp
--- a/Engine/Plugins/Tests/RHITests/Source/RHITests/Private/RHIReservedResourceTests.cpp
+++ b/Engine/Plugins/Tests/RHITests/Source/RHITests/Private/RHIReservedResourceTests.cpp
@@ -6,6 +6,75 @@
 #include "RenderCaptureInterface.h"
 #include "RHIStaticStates.h"
 
+bool FRHIReservedResourceTests::Test_ReservedResource_CreateVolumeTexture(FRHICommandListImmediate& RHICmdList)
+{
+	if (!GRHIGlobals.ReservedResources.SupportsVolumeTextures)
+	{
+		return true;
+	}
+
+	// Tiny volume texture with immediately committed physical memory.
+	// Expected to exercise the packed mip layout case.
+
+	{
+		FRHITextureCreateDesc Desc = FRHITextureCreateDesc::Create3D(TEXT("TestTinyReservedTexture3D"));
+		Desc.SetFlags(TexCreate_ShaderResource | TexCreate_ReservedResource | TexCreate_ImmediateCommit)
+			.SetExtent(FIntPoint(8, 8))
+			.SetDepth(8)
+			.SetNumMips(1)
+			.SetFormat(PF_A32B32G32R32F)
+			.SetInitialState(ERHIAccess::SRVCompute);
+
+		FTextureRHIRef Texture = RHICreateTexture(Desc);
+	}
+
+	// Small size volume texture (16MB) with immediately committed physical memory.
+
+	{
+		FRHITextureCreateDesc Desc = FRHITextureCreateDesc::Create3D(TEXT("TestSmallReservedTexture3D"));
+		Desc.SetFlags(TexCreate_ShaderResource | TexCreate_ReservedResource | TexCreate_ImmediateCommit)
+			.SetExtent(FIntPoint(128, 128))
+			.SetDepth(64)
+			.SetNumMips(1)
+			.SetFormat(PF_A32B32G32R32F)
+			.SetInitialState(ERHIAccess::SRVCompute);
+
+		FTextureRHIRef Texture = RHICreateTexture(Desc);
+	}
+
+	// Medium size volume texture (128MB) with immediately committed physical memory.
+	// Likely to be backed by multiple physical memory allocations.
+
+	{
+		FRHITextureCreateDesc Desc = FRHITextureCreateDesc::Create3D(TEXT("TestMediumReservedTexture3D"));
+		Desc.SetFlags(TexCreate_ShaderResource | TexCreate_ReservedResource | TexCreate_ImmediateCommit)
+			.SetExtent(FIntPoint(1024, 1024))
+			.SetDepth(8)
+			.SetNumMips(1)
+			.SetFormat(PF_A32B32G32R32F)
+			.SetInitialState(ERHIAccess::SRVCompute);
+
+		FTextureRHIRef Texture = RHICreateTexture(Desc);
+	}
+
+	// Huge volume texture without physical memory allocation.
+	// Just test the ability to reserve a huge virtual address range.
+
+	{
+		FRHITextureCreateDesc Desc = FRHITextureCreateDesc::Create3D(TEXT("TestHugeReservedTexture3D"));
+		Desc.SetFlags(TexCreate_ShaderResource | TexCreate_ReservedResource)
+			.SetExtent(FIntPoint(2048, 2048))
+			.SetDepth(2048)
+			.SetNumMips(1)
+			.SetFormat(PF_A32B32G32R32F)
+			.SetInitialState(ERHIAccess::SRVCompute);
+
+		FTextureRHIRef Texture = RHICreateTexture(Desc);
+	}
+
+	return true;
+}
+
 bool FRHIReservedResourceTests::Test_ReservedResource_CreateTexture(FRHICommandListImmediate& RHICmdList)
 {
 	if (!GRHIGlobals.ReservedResources.Supported)
@@ -83,6 +152,19 @@ bool FRHIReservedResourceTests::Test_ReservedResource_CreateTexture(FRHICommandL
 		}
 	}
 
+	// Try to create huge reserved textures without committing physical memory
+
+	{
+		FRHITextureCreateDesc Desc = FRHITextureCreateDesc::Create2D(TEXT("TestHugeReservedTexture2D"));
+		Desc.SetFlags(TexCreate_ShaderResource | TexCreate_ReservedResource)
+			.SetExtent(FIntPoint(16384, 16384))
+			.SetNumMips(1)
+			.SetFormat(PF_A32B32G32R32F)
+			.SetInitialState(ERHIAccess::SRVCompute);
+
+		FTextureRHIRef Texture = RHICreateTexture(Desc);
+	}
+
 	return true;
 }
 
diff --git a/Engine/Plugins/Tests/RHITests/Source/RHITests/Private/RHITests.spec.cpp b/Engine/Plugins/Tests/RHITests/Source/RHITests/Private/RHITests.spec.cpp
--- a/Engine/Plugins/Tests/RHITests/Source/RHITests/Private/RHITests.spec.cpp
+++ b/Engine/Plugins/Tests/RHITests/Source/RHITests/Private/RHITests.spec.cpp
@@ -137,6 +137,12 @@ void FAutomationRHITest::Define()
 
 	Describe("Test RHI Reserved Resource", [this]()
 	{
+		It("Create Reserved Volume Texture", [this]()
+		{
+			bool bResult = RunOnRenderThreadSynchronous(FRHIReservedResourceTests::Test_ReservedResource_CreateVolumeTexture);
+			TestEqual("Create Reserved Volume Texture failed", bResult, 1);
+		});
+
 		It("Create Reserved Texture", [this]()
 		{
 			bool bResult = RunOnRenderThreadSynchronous(FRHIReservedResourceTests::Test_ReservedResource_CreateTexture);
diff --git a/Engine/Plugins/Tests/RHITests/Source/RHITests/Private/RHITestsModule.cpp b/Engine/Plugins/Tests/RHITests/Source/RHITests/Private/RHITestsModule.cpp
--- a/Engine/Plugins/Tests/RHITests/Source/RHITests/Private/RHITestsModule.cpp
+++ b/Engine/Plugins/Tests/RHITests/Source/RHITests/Private/RHITestsModule.cpp
@@ -21,10 +21,12 @@ static bool RunTests_RenderThread(FRHICommandListImmediate& RHICmdList)
 	// Reserved resources
 	// ------------------------------------------------
 	{
+
 		RUN_TEST(FRHIReservedResourceTests::Test_ReservedResource_CommitBuffer(RHICmdList));
 		RUN_TEST(FRHIReservedResourceTests::Test_ReservedResource_DecommitBuffer(RHICmdList));
 		RUN_TEST(FRHIReservedResourceTests::Test_ReservedResource_CreateBuffer(RHICmdList));
 		RUN_TEST(FRHIReservedResourceTests::Test_ReservedResource_CreateTexture(RHICmdList));
+		RUN_TEST(FRHIReservedResourceTests::Test_ReservedResource_CreateVolumeTexture(RHICmdList));
 	}
 
 	// ------------------------------------------------
diff --git a/Engine/Plugins/Tests/RHITests/Source/RHITests/Public/RHIReservedResourceTests.h b/Engine/Plugins/Tests/RHITests/Source/RHITests/Public/RHIReservedResourceTests.h
--- a/Engine/Plugins/Tests/RHITests/Source/RHITests/Public/RHIReservedResourceTests.h
+++ b/Engine/Plugins/Tests/RHITests/Source/RHITests/Public/RHIReservedResourceTests.h
@@ -9,6 +9,7 @@ class FRHIReservedResourceTests
 public:
 
 	static bool Test_ReservedResource_CreateTexture(FRHICommandListImmediate& RHICmdList);
+	static bool Test_ReservedResource_CreateVolumeTexture(FRHICommandListImmediate& RHICmdList);
 	static bool Test_ReservedResource_CreateBuffer(FRHICommandListImmediate& RHICmdList);
 	static bool Test_ReservedResource_CommitBuffer(FRHICommandListImmediate& RHICmdList);
 	static bool Test_ReservedResource_DecommitBuffer(FRHICommandListImmediate& RHICmdList);
diff --git a/Engine/Source/Runtime/D3D12RHI/Private/D3D12Resources.cpp b/Engine/Source/Runtime/D3D12RHI/Private/D3D12Resources.cpp
--- a/Engine/Source/Runtime/D3D12RHI/Private/D3D12Resources.cpp
+++ b/Engine/Source/Runtime/D3D12RHI/Private/D3D12Resources.cpp
@@ -179,9 +179,15 @@ void FD3D12Resource::CommitReservedResource(ID3D12CommandQueue* D3DCommandQueue,
 
 	check(Desc.bReservedResource);
 	check(ReservedResourceData.IsValid());
-	checkf(Desc.Dimension == D3D12_RESOURCE_DIMENSION_BUFFER
-		|| Desc.Dimension == D3D12_RESOURCE_DIMENSION_TEXTURE2D, 
-		TEXT("CommitReservedResource is currently only implemented for 2D textures and buffers"));
+
+	checkf(GRHIGlobals.ReservedResources.Supported,
+		TEXT("Current RHI does not support reserved resources"));
+
+	if (Desc.Dimension == D3D12_RESOURCE_DIMENSION_TEXTURE3D)
+	{
+		checkf(GRHIGlobals.ReservedResources.SupportsVolumeTextures,
+			TEXT("Current RHI does not support reserved volume textures"));
+	}
 
 	if (Desc.Dimension != D3D12_RESOURCE_DIMENSION_BUFFER)
 	{
@@ -238,8 +244,7 @@ void FD3D12Resource::CommitReservedResource(ID3D12CommandQueue* D3DCommandQueue,
 	}
 	else if (PackedMipDesc.NumStandardMips != 0)
 	{
-		checkf(SubresourceTiling.DepthInTiles == 1, TEXT("3D reserved textures are not supported/implemented"));
-		NumStandardTilesPerSubresource = SubresourceTiling.WidthInTiles * SubresourceTiling.HeightInTiles;
+		NumStandardTilesPerSubresource = SubresourceTiling.WidthInTiles * SubresourceTiling.HeightInTiles * SubresourceTiling.DepthInTiles;
 		NumTotalTiles = NumStandardTilesPerSubresource * NumSubresources;
 	}
 	else // packed mip case
@@ -253,8 +258,9 @@ void FD3D12Resource::CommitReservedResource(ID3D12CommandQueue* D3DCommandQueue,
 		D3DResourceNumTiles, NumTotalTiles);
 
 	const uint32 NumRequiredCommitTiles = RequiredCommitSizeInBytes / TileSizeInBytes;
+	const uint32 NumTilesPerSlice = SubresourceTiling.WidthInTiles * SubresourceTiling.HeightInTiles;
 
-	auto GetTiledResourceCoordinate = [SubresourceTiling, NumStandardTilesPerSubresource, NumSubresources, MaxTilesPerHeap]
+	auto GetTiledResourceCoordinate = [SubresourceTiling, NumStandardTilesPerSubresource, NumSubresources, NumTilesPerSlice, MaxTilesPerHeap]
 		(uint32 OffsetInTiles, uint32 NumTiles) -> D3D12_TILED_RESOURCE_COORDINATE 
 	{
 		D3D12_TILED_RESOURCE_COORDINATE ResourceCoordinate = {}; // Coordinates are in tiles, not pixels
@@ -263,7 +269,7 @@ void FD3D12Resource::CommitReservedResource(ID3D12CommandQueue* D3DCommandQueue,
 			const uint32 TileIndexInSubresource = OffsetInTiles % NumStandardTilesPerSubresource;
 			ResourceCoordinate.X = TileIndexInSubresource % SubresourceTiling.WidthInTiles;
 			ResourceCoordinate.Y = (TileIndexInSubresource / SubresourceTiling.WidthInTiles) % SubresourceTiling.HeightInTiles;
-			ResourceCoordinate.Z = 0; // Only simple 2D / Array2D textures are implemented
+			ResourceCoordinate.Z = TileIndexInSubresource / NumTilesPerSlice;
 
 			ResourceCoordinate.Subresource = OffsetInTiles / NumStandardTilesPerSubresource;
 			check(ResourceCoordinate.Subresource <= NumSubresources);
@@ -271,7 +277,7 @@ void FD3D12Resource::CommitReservedResource(ID3D12CommandQueue* D3DCommandQueue,
 		else
 		{
 			// Packed mip level case:
-			// - Only simple 2D textures are expected (single subresource, no arrays)
+			// - Only simple textures are expected (single subresource, no arrays)
 			// - Entire packed mip level must be covered in one map operation, so mapping origin is always 0
 
 			checkf(NumSubresources == 1,
diff --git a/Engine/Source/Runtime/D3D12RHI/Private/D3D12Texture.cpp b/Engine/Source/Runtime/D3D12RHI/Private/D3D12Texture.cpp
--- a/Engine/Source/Runtime/D3D12RHI/Private/D3D12Texture.cpp
+++ b/Engine/Source/Runtime/D3D12RHI/Private/D3D12Texture.cpp
@@ -581,11 +581,13 @@ FD3D12ResourceDesc FD3D12DynamicRHI::GetResourceDesc(const FRHITextureDesc& Text
 	}
 
 	// Only 2D textures without mips are implemented/supported, to avoid the complexity associated with packed mips
-	if (GRHIGlobals.ReservedResources.Supported
-		&& EnumHasAllFlags(TextureDesc.Flags, TexCreate_ReservedResource)
-		&& TextureDesc.NumMips == 1
-		&& (TextureDesc.Dimension == ETextureDimension::Texture2D || TextureDesc.Dimension == ETextureDimension::Texture2DArray))
+	if (EnumHasAllFlags(TextureDesc.Flags, TexCreate_ReservedResource))
 	{
+		checkf(GRHIGlobals.ReservedResources.Supported, TEXT("Reserved resources resources are not supported on this machine"));
+		checkf(TextureDesc.NumMips == 1, TEXT("Reserved resources with mips are not supported"));
+		checkf(TextureDesc.IsTexture2D() || TextureDesc.IsTexture3D(), TEXT("Only 2D and 3D textures can be created as reserved resources"));
+		checkf(!TextureDesc.IsTexture3D() || GRHIGlobals.ReservedResources.SupportsVolumeTextures, TEXT("Current RHI does not support reserved volume textures"));
+
 		ResourceDesc.bReservedResource = true;
 		ResourceDesc.Layout = D3D12_TEXTURE_LAYOUT_64KB_UNDEFINED_SWIZZLE;
 	}
@@ -892,8 +894,29 @@ FD3D12Texture* FD3D12DynamicRHI::CreateD3D12Texture(const FRHITextureCreateDesc&
 		}
 		else if (CreateDesc.IsTexture3D())
 		{
-			VERIFYD3D12CREATETEXTURERESULT(Device->GetTextureAllocator().AllocateTexture(
-				ResourceDesc, ClearValuePtr, CreateDesc.Format, Location, CreateState, CreateDesc.DebugName), ResourceDesc, Device->GetDevice());
+			if (ResourceDesc.bReservedResource)
+			{
+				FD3D12Resource* Resource = nullptr;
+				VERIFYD3D12CREATETEXTURERESULT(
+					Adapter->CreateReservedResource(ResourceDesc, Device->GetGPUMask(), CreateState,
+						ED3D12ResourceStateMode::MultiState, CreateState, ClearValuePtr, &Resource, CreateDesc.DebugName, false),
+					ResourceDesc, Device->GetDevice());
+
+				D3D12_RESOURCE_ALLOCATION_INFO AllocInfo = Device->GetResourceAllocationInfo(ResourceDesc);
+
+				Location.AsStandAlone(Resource, AllocInfo.SizeInBytes);
+
+				if (EnumHasAllFlags(CreateDesc.Flags, TexCreate_ImmediateCommit))
+				{
+					Resource->CommitReservedResource(Device->GetQueue(ED3D12QueueType::Direct).D3DCommandQueue, UINT64_MAX /*commit entire resource*/);
+				}
+			}
+			else
+			{
+				VERIFYD3D12CREATETEXTURERESULT(Device->GetTextureAllocator().AllocateTexture(
+					ResourceDesc, ClearValuePtr, CreateDesc.Format, Location, CreateState, CreateDesc.DebugName), ResourceDesc, Device->GetDevice());
+			}
+
 			Location.SetOwner(NewTexture);
 		}
 		else
diff --git a/Engine/Source/Runtime/D3D12RHI/Private/Windows/WindowsD3D12Device.cpp b/Engine/Source/Runtime/D3D12RHI/Private/Windows/WindowsD3D12Device.cpp
--- a/Engine/Source/Runtime/D3D12RHI/Private/Windows/WindowsD3D12Device.cpp
+++ b/Engine/Source/Runtime/D3D12RHI/Private/Windows/WindowsD3D12Device.cpp
@@ -1641,6 +1641,9 @@ void FD3D12DynamicRHI::Init()
 		{
 			// Tier 2 is guaranteed for all adapters with feature level 12_0.
 			GRHIGlobals.ReservedResources.Supported = Options.TiledResourcesTier >= D3D12_TILED_RESOURCES_TIER_2;
+
+			// Tier 3 is required to create volume textures. Some hardware may support it.
+			GRHIGlobals.ReservedResources.SupportsVolumeTextures = Options.TiledResourcesTier >= D3D12_TILED_RESOURCES_TIER_3;
 		}
 	}
 
diff --git a/Engine/Source/Runtime/RHI/Private/RHI.cpp b/Engine/Source/Runtime/RHI/Private/RHI.cpp
--- a/Engine/Source/Runtime/RHI/Private/RHI.cpp
+++ b/Engine/Source/Runtime/RHI/Private/RHI.cpp
@@ -1726,11 +1726,20 @@ bool FRHITextureDesc::Validate(const FRHITextureCreateInfo& Desc, const TCHAR* N
 			TEXT("Reserved Texture %s's can't be created because current RHI does not support reserved resources."),
 			Name);
 
-		ValidateResourceDesc(
-			Desc.Dimension == ETextureDimension::Texture2D ||
-			Desc.Dimension == ETextureDimension::Texture2DArray,
-			TEXT("Reserved Texture %s's Desc.Dimension=%s is invalid. Expected Texture2D or Texture2DArray."),
-			Name, GetTextureDimensionString(Desc.Dimension));
+		if (Desc.IsTexture3D())
+		{
+			ValidateResourceDesc(GRHIGlobals.ReservedResources.SupportsVolumeTextures,
+				TEXT("Reserved Texture %s's can't be created because current RHI does not support reserved volume textures."),
+				Name);
+		}
+		else
+		{
+			ValidateResourceDesc(
+				Desc.Dimension == ETextureDimension::Texture2D ||
+				Desc.Dimension == ETextureDimension::Texture2DArray,
+				TEXT("Reserved Texture %s's Desc.Dimension=%s is invalid. Expected Texture2D, Texture2DArray or Texture3D."),
+				Name, GetTextureDimensionString(Desc.Dimension));
+		}
 
 		ValidateResourceDesc(Desc.NumMips == 1,
 			TEXT("Reserved Texture %s's NumMips=%d is invalid. Expected only 1 mip level."),
diff --git a/Engine/Source/Runtime/RHI/Public/RHIGlobals.h b/Engine/Source/Runtime/RHI/Public/RHIGlobals.h
--- a/Engine/Source/Runtime/RHI/Public/RHIGlobals.h
+++ b/Engine/Source/Runtime/RHI/Public/RHIGlobals.h
@@ -615,9 +615,15 @@ struct FRHIGlobals
 	{
 		/**
 		* True if the RHI supports reserved (AKA tiled, virtual or sparse) resources and operations related to them.
+		* Buffers and 2D textures (without mips) can be created with ReservedResource flag.
 		*/
 		bool Supported = false;
 
+		/**
+		* True if the RHI supports creating volume textures with ReservedResource flag.
+		*/
+		bool SupportsVolumeTextures = false;
+
 		/**
 		* Smallest mip dimension of reserved texture arrays must be greater or equal to this value.
 		* Tiled/reserved resources with both more than one array slice and any mipmap that
