workaround for rendertarget clear test on Vulkan for 2D arrays

#rnx

diff --git a/Engine/Plugins/Tests/RHITests/Source/RHITests/Private/RHIClearTextureTests.cpp b/Engine/Plugins/Tests/RHITests/Source/RHITests/Private/RHIClearTextureTests.cpp
--- a/Engine/Plugins/Tests/RHITests/Source/RHITests/Private/RHIClearTextureTests.cpp
+++ b/Engine/Plugins/Tests/RHITests/Source/RHITests/Private/RHIClearTextureTests.cpp
@@ -1432,9 +1432,10 @@ bool FRHIClearTextureTests::Test_ClearTexture(FRHICommandListImmediate& RHICmdLi
 	}
 
 	bool bSupportsSM5 = IsFeatureLevelSupported(GMaxRHIShaderPlatform, ERHIFeatureLevel::SM5);
+	bool bVulkanPlatform = IsVulkanPlatform(GMaxRHIShaderPlatform);
 
 	// Prune not only invalid combinations, but also try to keep the number of clear as low as possible to keep testing time acceptable
-	auto AddValidTextureDescIfValid = [bSupportsSM5](const FRHITextureCreateDesc& InRHITextureCreateDesc, TArray<FRHITextureCreateDesc>& OutRHITextureCreateDescs)
+	auto AddValidTextureDescIfValid = [bSupportsSM5, bVulkanPlatform](const FRHITextureCreateDesc& InRHITextureCreateDesc, TArray<FRHITextureCreateDesc>& OutRHITextureCreateDescs)
 	{
 		FRHITextureCreateDesc RHITextureCreateDesc = InRHITextureCreateDesc;
 		if (RHITextureCreateDesc.Dimension == ETextureDimension::Texture2D)
@@ -1482,6 +1483,8 @@ bool FRHIClearTextureTests::Test_ClearTexture(FRHICommandListImmediate& RHICmdLi
 				OutRHITextureCreateDescs.Add(RHITextureCreateDescCopy);
 			}
 
+			
+			if (bVulkanPlatform) return; // Workaround for VK not clearing 2D arrays properly. Follow up in UE-208082
 			if (!bSupportsSM5) return; // non sm5 platforms do not support write to multiple slices
 		}
 
