FShaderParameters should be bound through utility functions only to handle bindless mode correctly.

diff --git a/Engine/Plugins/Tests/RHITests/Source/RHITests/Private/RHIClearTextureTests.cpp b/Engine/Plugins/Tests/RHITests/Source/RHITests/Private/RHIClearTextureTests.cpp
--- a/Engine/Plugins/Tests/RHITests/Source/RHITests/Private/RHIClearTextureTests.cpp
+++ b/Engine/Plugins/Tests/RHITests/Source/RHITests/Private/RHIClearTextureTests.cpp
@@ -603,16 +603,16 @@ static void DrawTextureTo2DAtlas(FRHICommandListImmediate& RHICmdList, FRHITextu
 
 	SetComputePipelineState(RHICmdList, ComputeShader.GetComputeShader());
 	FRHIBatchedShaderParameters ShaderParameters;
-	ShaderParameters.SetShaderParameter(ComputeShader->MipBiasMipNumsViewport.GetBufferIndex(), ComputeShader->MipBiasMipNumsViewport.GetBaseIndex(), sizeof(MipBiasMipNumsViewport), &MipBiasMipNumsViewport);
+	SetShaderValue(ShaderParameters, ComputeShader->MipBiasMipNumsViewport, MipBiasMipNumsViewport);
 	if (GUseCustomSRV)
 	{
-		ShaderParameters.SetShaderResourceViewParameter(ComputeShader->SrcResourceParam.GetBaseIndex(), SourceSRV);
+		SetSRVParameter(ShaderParameters, ComputeShader->SrcResourceParam, SourceSRV);
 	}
 	else
 	{
-		ShaderParameters.SetShaderTexture(ComputeShader->SrcResourceParam.GetBaseIndex(), SourceTexture); // this doesn't work for cubemaps/cubemaps arrays when doing .Load on a Texture2DArray
+		SetTextureParameter(ShaderParameters, ComputeShader->SrcResourceParam, SourceTexture); // this doesn't work for cubemaps/cubemaps arrays when doing .Load on a Texture2DArray
 	}
-	ShaderParameters.SetUAVParameter(ComputeShader->RWAtlas2D.GetBaseIndex(), DestTextureUAV);
+	SetUAVParameter(ShaderParameters, ComputeShader->RWAtlas2D, DestTextureUAV);
 	RHICmdList.SetBatchedShaderParameters(ComputeShader.GetComputeShader(), ShaderParameters);
 	RHICmdList.DispatchComputeShader((AtlasViewport.X + 7) / 8, (AtlasViewport.Y + 7) / 8, 1);
 }
@@ -831,7 +831,7 @@ static void DrawColoredQuad(FRHICommandListImmediate& RHICmdList, FRHIVertexDecl
 			SetGraphicsPipelineState(RHICmdList, GraphicsPSOInit, 0);
 
 			FRHIBatchedShaderParameters ShaderParameters;
-			ShaderParameters.SetShaderParameter(PixelShader->TestFillTextureConstant.GetBufferIndex(), PixelShader->TestFillTextureConstant.GetBaseIndex(), sizeof(ConstColor), &ConstColor);
+			SetShaderValue(ShaderParameters, PixelShader->TestFillTextureConstant, ConstColor);
 			RHICmdList.SetBatchedShaderParameters(PixelShader.GetPixelShader(), ShaderParameters);
 
 			if (bAllSlicesInRTV)
diff --git a/Engine/Plugins/Tests/RHITests/Source/RHITests/Private/RHIDrawTests.cpp b/Engine/Plugins/Tests/RHITests/Source/RHITests/Private/RHIDrawTests.cpp
--- a/Engine/Plugins/Tests/RHITests/Source/RHITests/Private/RHIDrawTests.cpp
+++ b/Engine/Plugins/Tests/RHITests/Source/RHITests/Private/RHIDrawTests.cpp
@@ -34,7 +34,7 @@ class FTestDrawInstancedPS : public FGlobalShader
 	static void ModifyCompilationEnvironment(const FGlobalShaderPermutationParameters& Parameters, FShaderCompilerEnvironment& OutEnvironment)
 	{
 	}
-	LAYOUT_FIELD(FShaderParameter, OutDrawnInstances);
+	LAYOUT_FIELD(FShaderResourceParameter, OutDrawnInstances);
 };
 IMPLEMENT_GLOBAL_SHADER(FTestDrawInstancedPS, "/Plugin/RHITests/Private/TestDrawInstanced.usf", "TestDrawInstancedMainPS", SF_Pixel);
 
@@ -204,7 +204,7 @@ bool FRHIDrawTests::InternalDrawBaseVertexAndInstance(FRHICommandListImmediate&
 	RHICmdList.SetStreamSource(1, Resources.InstanceIDBuffer, 0);
 
 	FRHIBatchedShaderParameters ShaderParameters;
-	ShaderParameters.SetUAVParameter(Resources.PixelShader->OutDrawnInstances.GetBaseIndex(), Resources.OutputBufferUAV);
+	SetUAVParameter(ShaderParameters, Resources.PixelShader->OutDrawnInstances, Resources.OutputBufferUAV);
 	RHICmdList.SetBatchedShaderParameters(Resources.PixelShader.GetPixelShader(), ShaderParameters);
 
 	if (DrawKind == EDrawKind::Direct)
@@ -313,7 +313,7 @@ bool FRHIDrawTests::Test_MultiDrawIndirect(FRHICommandListImmediate& RHICmdList)
 	RHICmdList.SetStreamSource(1, Resources.InstanceIDBuffer, 0);
 
 	FRHIBatchedShaderParameters ShaderParameters;
-	ShaderParameters.SetUAVParameter(Resources.PixelShader->OutDrawnInstances.GetBaseIndex(), Resources.OutputBufferUAV);
+	SetUAVParameter(ShaderParameters, Resources.PixelShader->OutDrawnInstances, Resources.OutputBufferUAV);
 	RHICmdList.SetBatchedShaderParameters(Resources.PixelShader.GetPixelShader(), ShaderParameters);
 
 	const uint32 DrawArgsStride = sizeof(DrawArgs[0]);
