Fix framebuffer creation on vulkan to be able to specify the correct number of layers in a texture2d array
Re-enable Clear 2D array for all slices on vulkan

diff --git a/Engine/Plugins/Tests/RHITests/Source/RHITests/Private/RHIClearTextureTests.cpp b/Engine/Plugins/Tests/RHITests/Source/RHITests/Private/RHIClearTextureTests.cpp
--- a/Engine/Plugins/Tests/RHITests/Source/RHITests/Private/RHIClearTextureTests.cpp
+++ b/Engine/Plugins/Tests/RHITests/Source/RHITests/Private/RHIClearTextureTests.cpp
@@ -1432,10 +1432,9 @@ bool FRHIClearTextureTests::Test_ClearTexture(FRHICommandListImmediate& RHICmdLi
 	}
 
 	bool bSupportsSM5 = IsFeatureLevelSupported(GMaxRHIShaderPlatform, ERHIFeatureLevel::SM5);
-	bool bVulkanPlatform = IsVulkanPlatform(GMaxRHIShaderPlatform);
 
 	// Prune not only invalid combinations, but also try to keep the number of clear as low as possible to keep testing time acceptable
-	auto AddValidTextureDescIfValid = [bSupportsSM5, bVulkanPlatform](const FRHITextureCreateDesc& InRHITextureCreateDesc, TArray<FRHITextureCreateDesc>& OutRHITextureCreateDescs)
+	auto AddValidTextureDescIfValid = [bSupportsSM5](const FRHITextureCreateDesc& InRHITextureCreateDesc, TArray<FRHITextureCreateDesc>& OutRHITextureCreateDescs)
 	{
 		FRHITextureCreateDesc RHITextureCreateDesc = InRHITextureCreateDesc;
 		if (RHITextureCreateDesc.Dimension == ETextureDimension::Texture2D)
@@ -1482,9 +1481,7 @@ bool FRHIClearTextureTests::Test_ClearTexture(FRHICommandListImmediate& RHICmdLi
 				RHITextureCreateDescCopy.Flags |= TexCreate_TargetArraySlicesIndependently;
 				OutRHITextureCreateDescs.Add(RHITextureCreateDescCopy);
 			}
-
 			
-			if (bVulkanPlatform) return; // Workaround for VK not clearing 2D arrays properly. Follow up in UE-208082
 			if (!bSupportsSM5) return; // non sm5 platforms do not support write to multiple slices
 		}
 
@@ -1568,7 +1565,7 @@ bool FRHIClearTextureTests::Test_ClearTexture(FRHICommandListImmediate& RHICmdLi
 			}
 
 			CLEAR_TEST_SCOPED_NAMED_EVENT_F(TEXT("Test_ClearTexture %s"), FColor::Magenta, *GetTextureName(RHITextureCreateDescs[TextureIndex]));
-			CLEAR_TEST_SCOPED_DRAW_EVENTF(RHICmdList, ClearTextureTest, TEXT("Test_ClearTexture %s"), *GetTextureName(RHITextureCreateDescs[TextureIndex]));
+			CLEAR_TEST_SCOPED_DRAW_EVENTF(RHICmdList, ClearTextureTest, TEXT("Test_ClearTexture %s"), GetTextureName(RHITextureCreateDescs[TextureIndex]));
 
 			FTextureRHIRef Texture;
 			FShaderResourceViewRHIRef TextureSRV;
diff --git a/Engine/Source/Runtime/VulkanRHI/Private/VulkanViewport.cpp b/Engine/Source/Runtime/VulkanRHI/Private/VulkanViewport.cpp
--- a/Engine/Source/Runtime/VulkanRHI/Private/VulkanViewport.cpp
+++ b/Engine/Source/Runtime/VulkanRHI/Private/VulkanViewport.cpp
@@ -349,7 +349,8 @@ FVulkanFramebuffer::FVulkanFramebuffer(FVulkanDevice& Device, const FRHISetRende
 
 		if (Texture->GetViewType() == VK_IMAGE_VIEW_TYPE_2D || Texture->GetViewType() == VK_IMAGE_VIEW_TYPE_2D_ARRAY)
 		{
-			uint32 ArraySliceIndex, NumArraySlices;
+			uint32 ArraySliceIndex = 0;
+			uint32 NumArraySlices = 1;
 			if (InRTInfo.ColorRenderTarget[Index].ArraySliceIndex == -1)
 			{
 				ArraySliceIndex = 0;
@@ -362,6 +363,13 @@ FVulkanFramebuffer::FVulkanFramebuffer(FVulkanDevice& Device, const FRHISetRende
 				check(ArraySliceIndex < Texture->GetNumberOfArrayLevels());
 			}
 
+			// About !RTLayout.GetIsMultiView(), from https://registry.khronos.org/vulkan/specs/1.3-extensions/man/html/VkFramebufferCreateInfo.html: 
+			// If the render pass uses multiview, then layers must be one
+			if (Texture->GetViewType() == VK_IMAGE_VIEW_TYPE_2D_ARRAY && !RTLayout.GetIsMultiView())
+			{
+				NumLayers = NumArraySlices;
+			}
+
 			CreateOwnedView()->InitAsTextureView(
 				  Texture->Image
 				, Texture->GetViewType()
