Fix for RHI clear texture tests: 2 background tasks could call DispatchSubsequents on the same FGraphEventRef in a row, without resetting it first. Changing the code such that we only signal the event once when the first background task is over

diff --git a/Engine/Plugins/Tests/RHITests/Source/RHITests/Private/RHIClearTextureTests.cpp b/Engine/Plugins/Tests/RHITests/Source/RHITests/Private/RHIClearTextureTests.cpp
--- a/Engine/Plugins/Tests/RHITests/Source/RHITests/Private/RHIClearTextureTests.cpp
+++ b/Engine/Plugins/Tests/RHITests/Source/RHITests/Private/RHIClearTextureTests.cpp
@@ -185,23 +185,27 @@ FTextureRHIRef FStagingTexturePool::CreateTexture2DAtlasStagingFromPool(FRHIComm
 {
 	check(WaitForTextureEvent == nullptr);
 
+	FGraphEventRef WaitForTextureEventLocal;
 	{
 		FRWScopeLock WriteLock(TexturesToReturnToPoolRWLock, SLT_Write);
 		Recycle(RHICmdList);
 		if (AvailableStagingTextures[PixelFormat].Num() == 0)
 		{
-			WaitForTextureEvent = FGraphEvent::CreateGraphEvent();
+			check(WaitForTextureEvent == nullptr);
+			WaitForTextureEventLocal = FGraphEvent::CreateGraphEvent();
+            // Have the text task in-flight warn us when a new texture is ready
+			WaitForTextureEvent = WaitForTextureEventLocal;
 		}
 	}
 
-	if (WaitForTextureEvent)
+	if (WaitForTextureEventLocal)
 	{
 		CLEAR_TEST_SCOPED_NAMED_EVENT_TEXT("WaitForTexture", FColor::Magenta);
-		WaitForTextureEvent->Wait();
+		WaitForTextureEventLocal->Wait();
 		{
 			FRWScopeLock WriteLock(TexturesToReturnToPoolRWLock, SLT_Write);
 			Recycle(RHICmdList);
-			WaitForTextureEvent = nullptr;
+			check(WaitForTextureEvent == nullptr);
 		}
 	}
 
@@ -288,6 +292,8 @@ void FStagingTexturePool::ReturnToPool(const TArray<FTestOperation>& TestOperati
 	if (WaitForTextureEvent)
 	{
 		WaitForTextureEvent->DispatchSubsequents();
+        // We must call DispatchSubsequents only once
+		WaitForTextureEvent = nullptr;
 	}
 
 }
