Add RHI unit test for binding UAVs to graphics stages

- Test basic UAV write from pixel shader

diff --git a/Engine/Plugins/Tests/RHITests/Shaders/Private/TestGraphicsUAV.usf b/Engine/Plugins/Tests/RHITests/Shaders/Private/TestGraphicsUAV.usf
new file mode 100644
--- /dev/null
+++ b/Engine/Plugins/Tests/RHITests/Shaders/Private/TestGraphicsUAV.usf
@@ -0,0 +1,39 @@
+// Copyright Epic Games, Inc. All Rights Reserved.
+
+#include "/Engine/Private/Common.ush"
+
+RWBuffer<uint>	RWPixelShaderOutput;
+RWBuffer<uint>	RWVertexShaderOutput;
+
+void TestGraphicsUAVWriteMainVS(
+	in float4 VertexPosition : ATTRIBUTE0,
+	in uint VertexID : SV_VertexID,
+	in uint InstanceID : SV_InstanceID,
+	out uint OutInstanceID : Texcoord0,
+	out float4 OutPosition : SV_Position
+)
+{
+	RWVertexShaderOutput[VertexID] = VertexID;
+	OutPosition = VertexPosition;
+	OutInstanceID = InstanceID;
+}
+
+void TestGraphicsUAVTrivialMainVS(
+	in float4 VertexPosition : ATTRIBUTE0,
+	in uint InstanceID : SV_InstanceID,
+	out uint OutInstanceID : Texcoord0,
+	out float4 OutPosition : SV_Position
+)
+{
+	OutPosition = VertexPosition;
+	OutInstanceID = InstanceID;
+}
+
+void TestGraphicsUAVWriteMainPS(
+	in nointerpolation uint InstanceID : Texcoord0,
+	out float4 OutColor	: SV_Target0
+)
+{
+	RWPixelShaderOutput[InstanceID] = InstanceID;
+	OutColor = float4(1,1,1,1);
+}
diff --git a/Engine/Plugins/Tests/RHITests/Source/RHITests/Private/RHIDrawTests.cpp b/Engine/Plugins/Tests/RHITests/Source/RHITests/Private/RHIDrawTests.cpp
--- a/Engine/Plugins/Tests/RHITests/Source/RHITests/Private/RHIDrawTests.cpp
+++ b/Engine/Plugins/Tests/RHITests/Source/RHITests/Private/RHIDrawTests.cpp
@@ -41,25 +41,6 @@ IMPLEMENT_GLOBAL_SHADER(FTestDrawInstancedPS, "/Plugin/RHITests/Private/TestDraw
 namespace
 {
 
-template <typename T>
-static FBufferRHIRef CreateBufferWithData(EBufferUsageFlags UsageFlags, ERHIAccess ResourceState, const TCHAR* Name, TConstArrayView<T> Data)
-{
-	FRHICommandListBase& RHICmdList = FRHICommandListImmediate::Get();
-	uint32 BufferSize = sizeof(T) * Data.Num();
-	FRHIResourceCreateInfo CreateInfo(Name);
-	FBufferRHIRef Buffer = RHICmdList.CreateBuffer(BufferSize, UsageFlags, sizeof(T), ResourceState, CreateInfo);
-	void* MappedData = RHICmdList.LockBuffer(Buffer, 0, BufferSize, EResourceLockMode::RLM_WriteOnly);
-	FMemory::Memcpy(MappedData, Data.GetData(), BufferSize);
-	RHICmdList.UnlockBuffer(Buffer);
-	return Buffer;
-}
-
-template<typename T>
-static FBufferRHIRef CreateBufferWithData(EBufferUsageFlags UsageFlags, ERHIAccess ResourceState, const TCHAR* Name, TArrayView<T> Data)
-{
-	return CreateBufferWithData(UsageFlags, ResourceState, Name, TConstArrayView<T>(Data.GetData(), Data.Num()));
-}
-
 // Structure to initialize common resources required for various draw tests
 struct FDrawTestResources
 {
diff --git a/Engine/Plugins/Tests/RHITests/Source/RHITests/Private/RHIGraphicsUAVTests.cpp b/Engine/Plugins/Tests/RHITests/Source/RHITests/Private/RHIGraphicsUAVTests.cpp
new file mode 100644
--- /dev/null
+++ b/Engine/Plugins/Tests/RHITests/Source/RHITests/Private/RHIGraphicsUAVTests.cpp
@@ -0,0 +1,147 @@
+// Copyright Epic Games, Inc. All Rights Reserved.
+
+#include "RHIGraphicsUAVTests.h"
+#include "RHIBufferTests.h" // for VerifyBufferContents
+#include "RHIStaticStates.h"
+#include "ShaderCompilerCore.h"
+#include "GlobalShader.h"
+#include "PipelineStateCache.h"
+
+class FTestGraphicsUAVTrivialVS : public FGlobalShader
+{
+	DECLARE_GLOBAL_SHADER(FTestGraphicsUAVTrivialVS);
+	FTestGraphicsUAVTrivialVS() = default;
+	FTestGraphicsUAVTrivialVS(const ShaderMetaType::CompiledShaderInitializerType& Initializer)
+	: FGlobalShader(Initializer)
+	{
+	}
+	static void ModifyCompilationEnvironment(const FGlobalShaderPermutationParameters& Parameters, FShaderCompilerEnvironment& OutEnvironment)
+	{
+	}
+};
+IMPLEMENT_GLOBAL_SHADER(FTestGraphicsUAVTrivialVS, "/Plugin/RHITests/Private/TestGraphicsUAV.usf", "TestGraphicsUAVTrivialMainVS", SF_Vertex);
+
+class FTestGraphicsUAVWriteVS : public FGlobalShader
+{
+	DECLARE_GLOBAL_SHADER(FTestGraphicsUAVWriteVS);
+	FTestGraphicsUAVWriteVS() = default;
+	FTestGraphicsUAVWriteVS(const ShaderMetaType::CompiledShaderInitializerType& Initializer)
+	: FGlobalShader(Initializer)
+	{
+		RWVertexShaderOutput.Bind(Initializer.ParameterMap, TEXT("RWVertexShaderOutput"), SPF_Mandatory);
+	}
+	static void ModifyCompilationEnvironment(const FGlobalShaderPermutationParameters& Parameters, FShaderCompilerEnvironment& OutEnvironment)
+	{
+	}
+	LAYOUT_FIELD(FShaderResourceParameter, RWVertexShaderOutput);
+};
+IMPLEMENT_GLOBAL_SHADER(FTestGraphicsUAVWriteVS, "/Plugin/RHITests/Private/TestGraphicsUAV.usf", "TestGraphicsUAVWriteMainVS", SF_Vertex);
+
+class FTestGraphicsUAVWritePS : public FGlobalShader
+{
+	DECLARE_GLOBAL_SHADER(FTestGraphicsUAVWritePS);
+	FTestGraphicsUAVWritePS() = default;
+	FTestGraphicsUAVWritePS(const ShaderMetaType::CompiledShaderInitializerType& Initializer)
+	: FGlobalShader(Initializer)
+	{
+		RWPixelShaderOutput.Bind(Initializer.ParameterMap, TEXT("RWPixelShaderOutput"), SPF_Mandatory);
+	}
+	static void ModifyCompilationEnvironment(const FGlobalShaderPermutationParameters& Parameters, FShaderCompilerEnvironment& OutEnvironment)
+	{
+	}
+	LAYOUT_FIELD(FShaderResourceParameter, RWPixelShaderOutput);
+};
+IMPLEMENT_GLOBAL_SHADER(FTestGraphicsUAVWritePS, "/Plugin/RHITests/Private/TestGraphicsUAV.usf", "TestGraphicsUAVWriteMainPS", SF_Pixel);
+
+bool FRHIGraphicsUAVTests::Test_GraphicsUAV_PixelShader(FRHICommandListImmediate& RHICmdList)
+{
+	if (!GRHIGlobals.SupportsPixelShaderUAVs)
+	{
+		return true;
+	}
+
+	TShaderMapRef<FTestGraphicsUAVTrivialVS> VertexShader(GetGlobalShaderMap(GMaxRHIFeatureLevel));
+	TShaderMapRef<FTestGraphicsUAVWritePS> PixelShader(GetGlobalShaderMap(GMaxRHIFeatureLevel));
+
+	FIntPoint RenderTargetSize = FIntPoint(4, 4);
+	FTextureRHIRef RenderTarget;
+
+	FGraphicsPipelineStateInitializer GraphicsPSOInit;
+
+	FRHITextureDesc RenderTargetTextureDesc(ETextureDimension::Texture2D,
+		ETextureCreateFlags::RenderTargetable, PF_B8G8R8A8, FClearValueBinding(),
+		RenderTargetSize, 1, 1, 1, 1, 0);
+
+	FRHITextureCreateDesc RenderTargetCreateDesc(RenderTargetTextureDesc,
+		ERHIAccess::RTV, TEXT("GraphicsUAVTests_RenderTarget"));
+	RenderTarget = RHICreateTexture(RenderTargetCreateDesc);
+
+	FVertexDeclarationElementList VertexDeclarationElements;
+	VertexDeclarationElements.Add(FVertexElement(0, 0, VET_Float4, 0, 16));
+
+	FVertexDeclarationRHIRef VertexDeclarationRHI = PipelineStateCache::GetOrCreateVertexDeclaration(VertexDeclarationElements);
+
+	GraphicsPSOInit.BoundShaderState.VertexShaderRHI = VertexShader.GetVertexShader();
+	GraphicsPSOInit.BoundShaderState.VertexDeclarationRHI = VertexDeclarationRHI;
+	GraphicsPSOInit.BoundShaderState.PixelShaderRHI = PixelShader.GetPixelShader();
+	GraphicsPSOInit.DepthStencilState = TStaticDepthStencilState<false, CF_Always>::GetRHI();
+	GraphicsPSOInit.BlendState = TStaticBlendState<>::GetRHI();
+	GraphicsPSOInit.RasterizerState = TStaticRasterizerState<>::GetRHI();
+	GraphicsPSOInit.PrimitiveType = EPrimitiveType::PT_TriangleList;
+
+	TArray<FVector4f> Vertices;
+	Vertices.Reserve(3);
+	Vertices.Add(FVector4f(-1.0f, -1.0f, 0.0f, 1.0f));
+	Vertices.Add(FVector4f(-1.0f, +3.0f, 0.0f, 1.0f));
+	Vertices.Add(FVector4f(+3.0f, -1.0f, 0.0f, 1.0f));
+	
+	FBufferRHIRef VertexBuffer = CreateBufferWithData(EBufferUsageFlags::VertexBuffer,
+		ERHIAccess::VertexOrIndexBuffer, TEXT("GraphicsUAVTests_VertexBuffer"), MakeArrayView(Vertices));
+
+	static constexpr uint32 MaxInstances = 8;
+	static constexpr uint32 OutputBufferStride = sizeof(uint32);
+	static constexpr uint32 OutputBufferSize = OutputBufferStride * MaxInstances;
+	FRHIResourceCreateInfo OutputBufferCreateInfo(TEXT("GraphicsUAVTests_PixelShaderOutput"));
+
+	FBufferRHIRef OutputBuffer = RHICmdList.CreateBuffer(OutputBufferSize,
+		EBufferUsageFlags::UnorderedAccess | EBufferUsageFlags::SourceCopy,
+		OutputBufferStride, ERHIAccess::UAVCompute, OutputBufferCreateInfo);
+
+	FUnorderedAccessViewRHIRef OutputBufferUAV = RHICmdList.CreateUnorderedAccessView(OutputBuffer,
+		FRHIViewDesc::CreateBufferUAV()
+			.SetType(FRHIViewDesc::EBufferType::Typed)
+			.SetFormat(PF_R32_UINT));
+
+	RHICmdList.ClearUAVUint(OutputBufferUAV, FUintVector4(~0u));
+
+	FRHITexture* ColorRTs[1] = { RenderTarget.GetReference() };
+	FRHIRenderPassInfo RenderPassInfo(1, ColorRTs, ERenderTargetActions::DontLoad_DontStore);
+
+	RHICmdList.Transition(FRHITransitionInfo(OutputBufferUAV, ERHIAccess::UAVCompute, ERHIAccess::UAVGraphics, EResourceTransitionFlags::None));
+
+	RHICmdList.BeginRenderPass(RenderPassInfo, TEXT("GraphicsUAV_PixelShader"));
+	RHICmdList.SetViewport(0, 0, 0, float(RenderTargetSize.X), float(RenderTargetSize.Y), 1);
+
+	RHICmdList.ApplyCachedRenderTargets(GraphicsPSOInit);
+	SetGraphicsPipelineState(RHICmdList, GraphicsPSOInit, 0);
+
+	RHICmdList.SetStreamSource(0, VertexBuffer, 0);
+
+	FRHIBatchedShaderParameters ShaderParameters;
+	SetUAVParameter(ShaderParameters, PixelShader->RWPixelShaderOutput, OutputBufferUAV);
+	RHICmdList.SetBatchedShaderParameters(PixelShader.GetPixelShader(), ShaderParameters);
+
+	RHICmdList.DrawPrimitive(0, 1, MaxInstances);
+
+	RHICmdList.EndRenderPass();
+
+	RHICmdList.Transition(FRHITransitionInfo(OutputBufferUAV, ERHIAccess::UAVGraphics, ERHIAccess::CopySrc, EResourceTransitionFlags::None));
+
+	// Expect pixel shader to populate UAV with instance IDs
+
+	const uint32 ExpectedOutput[MaxInstances] = { 0, 1, 2, 3, 4, 5, 6, 7 };
+	TConstArrayView<uint8> ExpectedOutputView = MakeArrayView(reinterpret_cast<const uint8*>(ExpectedOutput), sizeof(ExpectedOutput));
+	const bool bSucceeded = FRHIBufferTests::VerifyBufferContents(TEXT("GraphicsUAV_PixelShader"), RHICmdList, OutputBuffer, ExpectedOutputView);
+
+	return bSucceeded;
+}
diff --git a/Engine/Plugins/Tests/RHITests/Source/RHITests/Private/RHITests.spec.cpp b/Engine/Plugins/Tests/RHITests/Source/RHITests/Private/RHITests.spec.cpp
--- a/Engine/Plugins/Tests/RHITests/Source/RHITests/Private/RHITests.spec.cpp
+++ b/Engine/Plugins/Tests/RHITests/Source/RHITests/Private/RHITests.spec.cpp
@@ -10,6 +10,7 @@
 #include "RHIClearTextureTests.h"
 #include "RHIReadbackTests.h"
 #include "RHIReservedResourceTests.h"
+#include "RHIGraphicsUAVTests.h"
 
 BEGIN_DEFINE_SPEC(FAutomationRHITest, "Rendering.RHI", EAutomationTestFlags::EngineFilter | EAutomationTestFlags::ApplicationContextMask | EAutomationTestFlags::NonNullRHI)
 END_DEFINE_SPEC(FAutomationRHITest)
@@ -174,4 +175,14 @@ void FAutomationRHITest::Define()
 			TestEqual("Decommit Reserved Buffer failed", bResult, 1);
 		});
 	});
+
+
+	Describe("Test RHI Graphics UAV Binding", [this]()
+	{
+		It("Pixel shader UAV", [this]()
+		{
+			bool bResult = RunOnRenderThreadSynchronous(FRHIGraphicsUAVTests::Test_GraphicsUAV_PixelShader);
+			TestEqual("Pixel shader UAV failed", bResult, 1);
+		});
+	});
 }
diff --git a/Engine/Plugins/Tests/RHITests/Source/RHITests/Private/RHITestsModule.cpp b/Engine/Plugins/Tests/RHITests/Source/RHITests/Private/RHITestsModule.cpp
--- a/Engine/Plugins/Tests/RHITests/Source/RHITests/Private/RHITestsModule.cpp
+++ b/Engine/Plugins/Tests/RHITests/Source/RHITests/Private/RHITestsModule.cpp
@@ -11,6 +11,7 @@
 #include "RHIClearTextureTests.h"
 #include "RHIReadbackTests.h"
 #include "RHIReservedResourceTests.h"
+#include "RHIGraphicsUAVTests.h"
 
 #define LOCTEXT_NAMESPACE "FRHITestsModule"
 
@@ -93,6 +94,13 @@ static bool RunTests_RenderThread(FRHICommandListImmediate& RHICmdList)
 	{
 		RUN_TEST(FRHIClearTextureTests::Test_ClearTexture(RHICmdList));
 	}
+	
+	// ------------------------------------------------
+	// Graphics UAV binding
+	// ------------------------------------------------
+	{
+		RUN_TEST(FRHIGraphicsUAVTests::Test_GraphicsUAV_PixelShader(RHICmdList));
+	}	
 
 	// @todo - add more tests
 	return bResult;
diff --git a/Engine/Plugins/Tests/RHITests/Source/RHITests/Public/RHIGraphicsUAVTests.h b/Engine/Plugins/Tests/RHITests/Source/RHITests/Public/RHIGraphicsUAVTests.h
new file mode 100644
--- /dev/null
+++ b/Engine/Plugins/Tests/RHITests/Source/RHITests/Public/RHIGraphicsUAVTests.h
@@ -0,0 +1,17 @@
+// Copyright Epic Games, Inc. All Rights Reserved.
+
+#pragma once
+
+#include "RHITestsCommon.h"
+
+/**
+ * Set of tests for binding UAVs to non-compute shaders
+ */
+
+class FRHIGraphicsUAVTests
+{
+public:
+
+	static bool Test_GraphicsUAV_PixelShader(FRHICommandListImmediate& RHICmdList);
+};
+
diff --git a/Engine/Plugins/Tests/RHITests/Source/RHITests/Public/RHITestsCommon.h b/Engine/Plugins/Tests/RHITests/Source/RHITests/Public/RHITestsCommon.h
--- a/Engine/Plugins/Tests/RHITests/Source/RHITests/Public/RHITestsCommon.h
+++ b/Engine/Plugins/Tests/RHITests/Source/RHITests/Public/RHITestsCommon.h
@@ -25,4 +25,23 @@ static inline FString ClearValueToString(const ValueType & ClearValue)
 	{
 		return FString::Printf(TEXT("0x%08x 0x%08x 0x%08x 0x%08x"), ClearValue.X, ClearValue.Y, ClearValue.Z, ClearValue.W);
 	}
-}
\ No newline at end of file
+}
+
+template <typename T>
+static FBufferRHIRef CreateBufferWithData(EBufferUsageFlags UsageFlags, ERHIAccess ResourceState, const TCHAR* Name, TConstArrayView<T> Data)
+{
+	FRHICommandListBase& RHICmdList = FRHICommandListImmediate::Get();
+	uint32 BufferSize = sizeof(T) * Data.Num();
+	FRHIResourceCreateInfo CreateInfo(Name);
+	FBufferRHIRef Buffer = RHICmdList.CreateBuffer(BufferSize, UsageFlags, sizeof(T), ResourceState, CreateInfo);
+	void* MappedData = RHICmdList.LockBuffer(Buffer, 0, BufferSize, EResourceLockMode::RLM_WriteOnly);
+	FMemory::Memcpy(MappedData, Data.GetData(), BufferSize);
+	RHICmdList.UnlockBuffer(Buffer);
+	return Buffer;
+}
+
+template<typename T>
+static FBufferRHIRef CreateBufferWithData(EBufferUsageFlags UsageFlags, ERHIAccess ResourceState, const TCHAR* Name, TArrayView<T> Data)
+{
+	return CreateBufferWithData(UsageFlags, ResourceState, Name, TConstArrayView<T>(Data.GetData(), Data.Num()));
+}
