Workaround for pixel shader UAV write unit test affecting some platforms

Platform shader compiler generates unsafe instructions under certain conditions when writing to a typed buffer. The workaround is to use a structured buffer, which avoids the problematic edge case.

diff --git a/Engine/Plugins/Tests/RHITests/Shaders/Private/TestGraphicsUAV.usf b/Engine/Plugins/Tests/RHITests/Shaders/Private/TestGraphicsUAV.usf
--- a/Engine/Plugins/Tests/RHITests/Shaders/Private/TestGraphicsUAV.usf
+++ b/Engine/Plugins/Tests/RHITests/Shaders/Private/TestGraphicsUAV.usf
@@ -2,8 +2,9 @@
 
 #include "/Engine/Private/Common.ush"
 
-RWBuffer<uint>	RWPixelShaderOutput;
-RWBuffer<uint>	RWVertexShaderOutput;
+// NOTE: using a structured buffer here as workaround for UE-212251
+RWStructuredBuffer<uint>	RWPixelShaderOutput;
+RWStructuredBuffer<uint>	RWVertexShaderOutput;
 
 void TestGraphicsUAVWriteMainVS(
 	in float4 VertexPosition : ATTRIBUTE0,
diff --git a/Engine/Plugins/Tests/RHITests/Source/RHITests/Private/RHIGraphicsUAVTests.cpp b/Engine/Plugins/Tests/RHITests/Source/RHITests/Private/RHIGraphicsUAVTests.cpp
--- a/Engine/Plugins/Tests/RHITests/Source/RHITests/Private/RHIGraphicsUAVTests.cpp
+++ b/Engine/Plugins/Tests/RHITests/Source/RHITests/Private/RHIGraphicsUAVTests.cpp
@@ -103,14 +103,15 @@ bool FRHIGraphicsUAVTests::Test_GraphicsUAV_PixelShader(FRHICommandListImmediate
 	static constexpr uint32 OutputBufferSize = OutputBufferStride * MaxInstances;
 	FRHIResourceCreateInfo OutputBufferCreateInfo(TEXT("GraphicsUAVTests_PixelShaderOutput"));
 
+	// NOTE: using a structured buffer here as workaround for UE-212251
 	FBufferRHIRef OutputBuffer = RHICmdList.CreateBuffer(OutputBufferSize,
-		EBufferUsageFlags::UnorderedAccess | EBufferUsageFlags::SourceCopy,
+		EBufferUsageFlags::UnorderedAccess | EBufferUsageFlags::SourceCopy | EBufferUsageFlags::StructuredBuffer,
 		OutputBufferStride, ERHIAccess::UAVCompute, OutputBufferCreateInfo);
 
 	FUnorderedAccessViewRHIRef OutputBufferUAV = RHICmdList.CreateUnorderedAccessView(OutputBuffer,
 		FRHIViewDesc::CreateBufferUAV()
-			.SetType(FRHIViewDesc::EBufferType::Typed)
-			.SetFormat(PF_R32_UINT));
+			.SetType(FRHIViewDesc::EBufferType::Structured)
+			.SetStride(4));
 
 	RHICmdList.ClearUAVUint(OutputBufferUAV, FUintVector4(~0u));
 
