Allow UAV binding in Vertex Shaders in D3D11 and D3D12 RHIs

diff --git a/Engine/Config/Windows/DataDrivenPlatformInfo.ini b/Engine/Config/Windows/DataDrivenPlatformInfo.ini
--- a/Engine/Config/Windows/DataDrivenPlatformInfo.ini
+++ b/Engine/Config/Windows/DataDrivenPlatformInfo.ini
@@ -75,6 +75,7 @@ FriendlyName=LOCTEXT("FriendlyShaderPlatformName_D3D_SM5", "PC D3D SM5")
 bSupportsDualSourceBlending=true
 bSupportsNNEShaders=true
 bSupportsIndependentSamplers=true
+bSupportsVertexShaderUAVs=RuntimeGuaranteed
 
 [PreviewPlatform PCD3D_SM5]
 PlatformName=PCD3D_SM5
@@ -139,6 +140,7 @@ bSupportsIndependentSamplers=true
 MaxSamplers=32
 SupportsBarycentricsSemantic=RuntimeDependent
 bSupportsWorkGraphs=true
+bSupportsVertexShaderUAVs=RuntimeGuaranteed
 
 [PreviewPlatform PCD3D_SM6]
 PlatformName=PCD3D_SM6
diff --git a/Engine/Plugins/Tests/RHITests/Source/RHITests/Private/RHIGraphicsUAVTests.cpp b/Engine/Plugins/Tests/RHITests/Source/RHITests/Private/RHIGraphicsUAVTests.cpp
--- a/Engine/Plugins/Tests/RHITests/Source/RHITests/Private/RHIGraphicsUAVTests.cpp
+++ b/Engine/Plugins/Tests/RHITests/Source/RHITests/Private/RHIGraphicsUAVTests.cpp
@@ -53,16 +53,10 @@ class FTestGraphicsUAVWritePS : public FGlobalShader
 };
 IMPLEMENT_GLOBAL_SHADER(FTestGraphicsUAVWritePS, "/Plugin/RHITests/Private/TestGraphicsUAV.usf", "TestGraphicsUAVWriteMainPS", SF_Pixel);
 
-bool FRHIGraphicsUAVTests::Test_GraphicsUAV_PixelShader(FRHICommandListImmediate& RHICmdList)
+static void Test_GraphicsUAV_Common(FRHICommandListImmediate& RHICmdList,
+	FRHIVertexShader* VertexShaderRHI, FRHIPixelShader* PixelShaderRHI,
+	TFunctionRef<void(FRHICommandListImmediate& RHICmdList)> RenderCallback)
 {
-	if (!GRHIGlobals.SupportsPixelShaderUAVs)
-	{
-		return true;
-	}
-
-	TShaderMapRef<FTestGraphicsUAVTrivialVS> VertexShader(GetGlobalShaderMap(GMaxRHIFeatureLevel));
-	TShaderMapRef<FTestGraphicsUAVWritePS> PixelShader(GetGlobalShaderMap(GMaxRHIFeatureLevel));
-
 	FIntPoint RenderTargetSize = FIntPoint(4, 4);
 	FTextureRHIRef RenderTarget;
 
@@ -81,9 +75,9 @@ bool FRHIGraphicsUAVTests::Test_GraphicsUAV_PixelShader(FRHICommandListImmediate
 
 	FVertexDeclarationRHIRef VertexDeclarationRHI = PipelineStateCache::GetOrCreateVertexDeclaration(VertexDeclarationElements);
 
-	GraphicsPSOInit.BoundShaderState.VertexShaderRHI = VertexShader.GetVertexShader();
+	GraphicsPSOInit.BoundShaderState.VertexShaderRHI = VertexShaderRHI;
 	GraphicsPSOInit.BoundShaderState.VertexDeclarationRHI = VertexDeclarationRHI;
-	GraphicsPSOInit.BoundShaderState.PixelShaderRHI = PixelShader.GetPixelShader();
+	GraphicsPSOInit.BoundShaderState.PixelShaderRHI = PixelShaderRHI;
 	GraphicsPSOInit.DepthStencilState = TStaticDepthStencilState<false, CF_Always>::GetRHI();
 	GraphicsPSOInit.BlendState = TStaticBlendState<>::GetRHI();
 	GraphicsPSOInit.RasterizerState = TStaticRasterizerState<>::GetRHI();
@@ -94,10 +88,36 @@ bool FRHIGraphicsUAVTests::Test_GraphicsUAV_PixelShader(FRHICommandListImmediate
 	Vertices.Add(FVector4f(-1.0f, -1.0f, 0.0f, 1.0f));
 	Vertices.Add(FVector4f(-1.0f, +3.0f, 0.0f, 1.0f));
 	Vertices.Add(FVector4f(+3.0f, -1.0f, 0.0f, 1.0f));
-	
+
 	FBufferRHIRef VertexBuffer = CreateBufferWithData(EBufferUsageFlags::VertexBuffer,
 		ERHIAccess::VertexOrIndexBuffer, TEXT("GraphicsUAVTests_VertexBuffer"), MakeArrayView(Vertices));
 
+	FRHITexture* ColorRTs[1] = { RenderTarget.GetReference() };
+	FRHIRenderPassInfo RenderPassInfo(1, ColorRTs, ERenderTargetActions::DontLoad_DontStore);
+
+	RHICmdList.BeginRenderPass(RenderPassInfo, TEXT("GraphicsUAVTest"));
+	RHICmdList.SetViewport(0, 0, 0, float(RenderTargetSize.X), float(RenderTargetSize.Y), 1);
+
+	RHICmdList.ApplyCachedRenderTargets(GraphicsPSOInit);
+	SetGraphicsPipelineState(RHICmdList, GraphicsPSOInit, 0);
+
+	RHICmdList.SetStreamSource(0, VertexBuffer, 0);
+
+	RenderCallback(RHICmdList);
+
+	RHICmdList.EndRenderPass();
+}
+
+bool FRHIGraphicsUAVTests::Test_GraphicsUAV_PixelShader(FRHICommandListImmediate& RHICmdList)
+{
+	if (!GRHIGlobals.SupportsPixelShaderUAVs)
+	{
+		return true;
+	}
+
+	TShaderMapRef<FTestGraphicsUAVTrivialVS> VertexShader(GetGlobalShaderMap(GMaxRHIFeatureLevel));
+	TShaderMapRef<FTestGraphicsUAVWritePS> PixelShader(GetGlobalShaderMap(GMaxRHIFeatureLevel));
+
 	static constexpr uint32 MaxInstances = 8;
 	static constexpr uint32 OutputBufferStride = sizeof(uint32);
 	static constexpr uint32 OutputBufferSize = OutputBufferStride * MaxInstances;
@@ -115,34 +135,76 @@ bool FRHIGraphicsUAVTests::Test_GraphicsUAV_PixelShader(FRHICommandListImmediate
 
 	RHICmdList.ClearUAVUint(OutputBufferUAV, FUintVector4(~0u));
 
-	FRHITexture* ColorRTs[1] = { RenderTarget.GetReference() };
-	FRHIRenderPassInfo RenderPassInfo(1, ColorRTs, ERenderTargetActions::DontLoad_DontStore);
-
 	RHICmdList.Transition(FRHITransitionInfo(OutputBufferUAV, ERHIAccess::UAVCompute, ERHIAccess::UAVGraphics, EResourceTransitionFlags::None));
 
-	RHICmdList.BeginRenderPass(RenderPassInfo, TEXT("GraphicsUAV_PixelShader"));
-	RHICmdList.SetViewport(0, 0, 0, float(RenderTargetSize.X), float(RenderTargetSize.Y), 1);
+	Test_GraphicsUAV_Common(RHICmdList,
+		VertexShader.GetVertexShader(),
+		PixelShader.GetPixelShader(),
+		[&PixelShader, &OutputBufferUAV](FRHICommandListImmediate& RHICmdList)
+	{
+			FRHIBatchedShaderParameters ShaderParameters;
+			SetUAVParameter(ShaderParameters, PixelShader->RWPixelShaderOutput, OutputBufferUAV);
+			RHICmdList.SetBatchedShaderParameters(PixelShader.GetPixelShader(), ShaderParameters);
+			RHICmdList.DrawPrimitive(0, 1, MaxInstances);
+	});
 
-	RHICmdList.ApplyCachedRenderTargets(GraphicsPSOInit);
-	SetGraphicsPipelineState(RHICmdList, GraphicsPSOInit, 0);
+	RHICmdList.Transition(FRHITransitionInfo(OutputBufferUAV, ERHIAccess::UAVGraphics, ERHIAccess::CopySrc, EResourceTransitionFlags::None));
 
-	RHICmdList.SetStreamSource(0, VertexBuffer, 0);
+	// Expect pixel shader to populate UAV with instance IDs
+
+	const uint32 ExpectedOutput[MaxInstances] = { 0, 1, 2, 3, 4, 5, 6, 7 };
+	TConstArrayView<uint8> ExpectedOutputView = MakeArrayView(reinterpret_cast<const uint8*>(ExpectedOutput), sizeof(ExpectedOutput));
+	const bool bSucceeded = FRHIBufferTests::VerifyBufferContents(TEXT("GraphicsUAV_PixelShader"), RHICmdList, OutputBuffer, ExpectedOutputView);
 
-	FRHIBatchedShaderParameters ShaderParameters;
-	SetUAVParameter(ShaderParameters, PixelShader->RWPixelShaderOutput, OutputBufferUAV);
-	RHICmdList.SetBatchedShaderParameters(PixelShader.GetPixelShader(), ShaderParameters);
+	return bSucceeded;
+}
 
-	RHICmdList.DrawPrimitive(0, 1, MaxInstances);
+bool FRHIGraphicsUAVTests::Test_GraphicsUAV_VertexShader(FRHICommandListImmediate& RHICmdList)
+{
+	if (!GRHIGlobals.SupportsVertexShaderUAVs)
+	{
+		return true;
+	}
 
-	RHICmdList.EndRenderPass();
+	TShaderMapRef<FTestGraphicsUAVWriteVS> VertexShader(GetGlobalShaderMap(GMaxRHIFeatureLevel));
+
+	static const uint32 MaxVertices = 3;
+	static constexpr uint32 OutputBufferStride = sizeof(uint32);
+	static constexpr uint32 OutputBufferSize = OutputBufferStride * MaxVertices;
+	FRHIResourceCreateInfo OutputBufferCreateInfo(TEXT("GraphicsUAVTests_VertexShaderOutput"));
+
+	FBufferRHIRef OutputBuffer = RHICmdList.CreateBuffer(OutputBufferSize,
+		EBufferUsageFlags::UnorderedAccess | EBufferUsageFlags::SourceCopy,
+		OutputBufferStride, ERHIAccess::UAVCompute, OutputBufferCreateInfo);
+
+	FUnorderedAccessViewRHIRef OutputBufferUAV = RHICmdList.CreateUnorderedAccessView(OutputBuffer,
+		FRHIViewDesc::CreateBufferUAV()
+		.SetType(FRHIViewDesc::EBufferType::Typed)
+		.SetFormat(PF_R32_UINT));
+
+	RHICmdList.ClearUAVUint(OutputBufferUAV, FUintVector4(~0u));
+
+	RHICmdList.Transition(FRHITransitionInfo(OutputBufferUAV, ERHIAccess::UAVCompute, ERHIAccess::UAVGraphics, EResourceTransitionFlags::None));
+
+	Test_GraphicsUAV_Common(RHICmdList,
+		VertexShader.GetVertexShader(),
+		nullptr, // vertex-only rendering
+	[&VertexShader, &OutputBufferUAV](FRHICommandListImmediate& RHICmdList)
+	{
+		FRHIBatchedShaderParameters ShaderParameters;
+		SetUAVParameter(ShaderParameters, VertexShader->RWVertexShaderOutput, OutputBufferUAV);
+		RHICmdList.SetBatchedShaderParameters(VertexShader.GetVertexShader(), ShaderParameters);
+		RHICmdList.DrawPrimitive(0, 1, 1);
+	});
 
 	RHICmdList.Transition(FRHITransitionInfo(OutputBufferUAV, ERHIAccess::UAVGraphics, ERHIAccess::CopySrc, EResourceTransitionFlags::None));
 
-	// Expect pixel shader to populate UAV with instance IDs
+	// Expect vertex shader to populate UAV with vertex IDs
 
-	const uint32 ExpectedOutput[MaxInstances] = { 0, 1, 2, 3, 4, 5, 6, 7 };
+	const uint32 ExpectedOutput[MaxVertices] = { 0, 1, 2 };
 	TConstArrayView<uint8> ExpectedOutputView = MakeArrayView(reinterpret_cast<const uint8*>(ExpectedOutput), sizeof(ExpectedOutput));
-	const bool bSucceeded = FRHIBufferTests::VerifyBufferContents(TEXT("GraphicsUAV_PixelShader"), RHICmdList, OutputBuffer, ExpectedOutputView);
+	const bool bSucceeded = FRHIBufferTests::VerifyBufferContents(TEXT("GraphicsUAV_VertexShader"), RHICmdList, OutputBuffer, ExpectedOutputView);
 
 	return bSucceeded;
 }
+
diff --git a/Engine/Plugins/Tests/RHITests/Source/RHITests/Private/RHITests.spec.cpp b/Engine/Plugins/Tests/RHITests/Source/RHITests/Private/RHITests.spec.cpp
--- a/Engine/Plugins/Tests/RHITests/Source/RHITests/Private/RHITests.spec.cpp
+++ b/Engine/Plugins/Tests/RHITests/Source/RHITests/Private/RHITests.spec.cpp
@@ -184,5 +184,11 @@ void FAutomationRHITest::Define()
 			bool bResult = RunOnRenderThreadSynchronous(FRHIGraphicsUAVTests::Test_GraphicsUAV_PixelShader);
 			TestEqual("Pixel shader UAV failed", bResult, 1);
 		});
+
+		It("Vertex shader UAV", [this]()
+		{
+			bool bResult = RunOnRenderThreadSynchronous(FRHIGraphicsUAVTests::Test_GraphicsUAV_VertexShader);
+			TestEqual("Vixel shader UAV failed", bResult, 1);
+		});
 	});
 }
diff --git a/Engine/Plugins/Tests/RHITests/Source/RHITests/Private/RHITestsModule.cpp b/Engine/Plugins/Tests/RHITests/Source/RHITests/Private/RHITestsModule.cpp
--- a/Engine/Plugins/Tests/RHITests/Source/RHITests/Private/RHITestsModule.cpp
+++ b/Engine/Plugins/Tests/RHITests/Source/RHITests/Private/RHITestsModule.cpp
@@ -100,7 +100,8 @@ static bool RunTests_RenderThread(FRHICommandListImmediate& RHICmdList)
 	// ------------------------------------------------
 	{
 		RUN_TEST(FRHIGraphicsUAVTests::Test_GraphicsUAV_PixelShader(RHICmdList));
-	}	
+		RUN_TEST(FRHIGraphicsUAVTests::Test_GraphicsUAV_VertexShader(RHICmdList));
+	}
 
 	// @todo - add more tests
 	return bResult;
diff --git a/Engine/Plugins/Tests/RHITests/Source/RHITests/Public/RHIGraphicsUAVTests.h b/Engine/Plugins/Tests/RHITests/Source/RHITests/Public/RHIGraphicsUAVTests.h
--- a/Engine/Plugins/Tests/RHITests/Source/RHITests/Public/RHIGraphicsUAVTests.h
+++ b/Engine/Plugins/Tests/RHITests/Source/RHITests/Public/RHIGraphicsUAVTests.h
@@ -13,5 +13,6 @@ class FRHIGraphicsUAVTests
 public:
 
 	static bool Test_GraphicsUAV_PixelShader(FRHICommandListImmediate& RHICmdList);
+	static bool Test_GraphicsUAV_VertexShader(FRHICommandListImmediate& RHICmdList);
 };
 
diff --git a/Engine/Shaders/Public/Platform.ush b/Engine/Shaders/Public/Platform.ush
--- a/Engine/Shaders/Public/Platform.ush
+++ b/Engine/Shaders/Public/Platform.ush
@@ -521,6 +521,12 @@
 #define PLATFORM_SUPPORTS_VERTEX_SHADER_SRVS 1
 #endif
 
+// Whether the platform supports binding UAVs to the vertex shader stage.
+#ifndef PLATFORM_SUPPORTS_VERTEX_SHADER_UAVS
+#define PLATFORM_SUPPORTS_VERTEX_SHADER_UAVS 0
+#endif
+
+
 // ---------------------------------------------------- Compiler specific defaults and fallbacks
 
 #if !defined(PLATFORM_BREAK)
diff --git a/Engine/Source/Runtime/D3D12RHI/Private/D3D12Commands.cpp b/Engine/Source/Runtime/D3D12RHI/Private/D3D12Commands.cpp
--- a/Engine/Source/Runtime/D3D12RHI/Private/D3D12Commands.cpp
+++ b/Engine/Source/Runtime/D3D12RHI/Private/D3D12Commands.cpp
@@ -1023,13 +1023,13 @@ static void SetShaderParametersOnContext(
 	{
 		if (Parameter.Type == FRHIShaderParameterResource::EType::UnorderedAccessView)
 		{
-			if (ShaderFrequency == SF_Pixel || ShaderFrequency == SF_Compute)
+			if (ShaderFrequency == SF_Pixel || ShaderFrequency == SF_Vertex|| ShaderFrequency == SF_Compute)
 			{
 				Binder.SetUAV(static_cast<FRHIUnorderedAccessView*>(Parameter.Resource), Parameter.Index, true);
 			}
 			else
 			{
-				checkf(false, TEXT("TShaderRHI Can't have compute shader to be set. UAVs are not supported on vertex, tessellation and geometry shaders."));
+				checkf(false, TEXT("TShaderRHI Can't have compute shader to be set. UAVs are not supported tessellation and geometry shaders."));
 			}
 		}
 	}
@@ -1114,13 +1114,13 @@ static void SetShaderUnbindsOnContext(
 			Context.StateCache.SetShaderResourceView(ShaderFrequency, nullptr, Unbind.Index);
 			break;
 		case FRHIShaderParameterUnbind::EType::UnorderedAccessView:
-			if (ShaderFrequency == SF_Pixel || ShaderFrequency == SF_Compute)
+			if (ShaderFrequency == SF_Pixel || ShaderFrequency == SF_Vertex || ShaderFrequency == SF_Compute)
 			{
 				Context.StateCache.SetUAV(ShaderFrequency, Unbind.Index, nullptr);
 			}
 			else
 			{
-				checkf(false, TEXT("TShaderRHI Can't have compute shader to be set. UAVs are not supported on vertex, tessellation and geometry shaders."));
+				checkf(false, TEXT("TShaderRHI Can't have compute shader to be set. UAVs are not supported tessellation and geometry shaders."));
 			}
 			break;
 		default:
diff --git a/Engine/Source/Runtime/D3D12RHI/Private/D3D12DescriptorCache.cpp b/Engine/Source/Runtime/D3D12RHI/Private/D3D12DescriptorCache.cpp
--- a/Engine/Source/Runtime/D3D12RHI/Private/D3D12DescriptorCache.cpp
+++ b/Engine/Source/Runtime/D3D12RHI/Private/D3D12DescriptorCache.cpp
@@ -278,10 +278,12 @@ D3D12_GPU_DESCRIPTOR_HANDLE FD3D12DescriptorCache::BuildUAVTable(EShaderFrequenc
 
 void FD3D12DescriptorCache::SetUAVTable(EShaderFrequency ShaderStage, const FD3D12RootSignature* RootSignature, FD3D12UnorderedAccessViewCache& Cache, uint32 SlotsNeeded, const D3D12_GPU_DESCRIPTOR_HANDLE& BindDescriptor)
 {
-	check(ShaderStage == SF_Compute || ShaderStage == SF_Pixel);
+	check(ShaderStage == SF_Compute || ShaderStage == SF_Pixel || ShaderStage == SF_Vertex);
 	const uint32 RootParameterIndex = RootSignature->UAVRDTBindSlot(ShaderStage);
 
-	if (ShaderStage == SF_Pixel)
+	check(FD3D12RootSignature::IsValidBindSlot(RootParameterIndex));
+
+	if (ShaderStage == SF_Pixel || ShaderStage == SF_Vertex)
 	{
 		Context.GraphicsCommandList()->SetGraphicsRootDescriptorTable(RootParameterIndex, BindDescriptor);
 	}
diff --git a/Engine/Source/Runtime/D3D12RHI/Private/D3D12RHI.cpp b/Engine/Source/Runtime/D3D12RHI/Private/D3D12RHI.cpp
--- a/Engine/Source/Runtime/D3D12RHI/Private/D3D12RHI.cpp
+++ b/Engine/Source/Runtime/D3D12RHI/Private/D3D12RHI.cpp
@@ -257,6 +257,10 @@ FD3D12DynamicRHI::FD3D12DynamicRHI(const TArray<TSharedPtr<FD3D12Adapter>>& Chos
 	GRHISupportsRHIOnTaskThread = true;
 
 	GRHIGlobals.NeedsShaderUnbinds = true;
+
+	// All D3D12 hardware supports binding UAVs to Vertex Shaders.
+	// Enable run-time support if corresponding bit is set in DDSPI.
+	GRHIGlobals.SupportsVertexShaderUAVs = true;
 }
 
 void FD3D12DynamicRHI::PostInit()
diff --git a/Engine/Source/Runtime/D3D12RHI/Private/D3D12RootSignature.cpp b/Engine/Source/Runtime/D3D12RHI/Private/D3D12RootSignature.cpp
--- a/Engine/Source/Runtime/D3D12RHI/Private/D3D12RootSignature.cpp
+++ b/Engine/Source/Runtime/D3D12RHI/Private/D3D12RootSignature.cpp
@@ -542,7 +542,7 @@ void FD3D12RootSignature::InternalAnalyzeSignature(const RootSignatureDescType&
 {
 	// Reset members to default values.
 	{
-		FMemory::Memset(BindSlotMap, 0xFF, sizeof(BindSlotMap));
+		FMemory::Memset(BindSlotMap, InvalidBindSlotMapIndex, sizeof(BindSlotMap));
 		bHasUAVs = false;
 		bHasSRVs = false;
 		bHasCBVs = false;
@@ -704,7 +704,7 @@ void FD3D12RootSignature::InternalAnalyzeSignature(const RootSignatureDescType&
 				UpdateCBVRegisterMaskWithDescriptor(CurrentVisibleSF, CurrentParameter.Descriptor);
 
 				// The first CBV for this stage must come first in the root signature, and subsequent root CBVs for this stage must be contiguous.
-				check(0xFF != CBVRDBindSlot(CurrentVisibleSF, 0));
+				check(InvalidBindSlotMapIndex != CBVRDBindSlot(CurrentVisibleSF, 0));
 				check(i == CBVRDBindSlot(CurrentVisibleSF, 0) + CurrentParameter.Descriptor.ShaderRegister);
 			}
 			break;
diff --git a/Engine/Source/Runtime/D3D12RHI/Private/D3D12RootSignature.h b/Engine/Source/Runtime/D3D12RHI/Private/D3D12RootSignature.h
--- a/Engine/Source/Runtime/D3D12RHI/Private/D3D12RootSignature.h
+++ b/Engine/Source/Runtime/D3D12RHI/Private/D3D12RootSignature.h
@@ -20,6 +20,7 @@ enum ERootParameterKeys
 	VS_CBVs,
 	VS_RootCBVs,
 	VS_Samplers,
+	VS_UAVs,
 	GS_SRVs,
 	GS_CBVs,
 	GS_RootCBVs,
@@ -36,7 +37,7 @@ enum ERootParameterKeys
 	ALL_CBVs,
 	ALL_RootCBVs,
 	ALL_Samplers,
-	ALL_UAVs,
+	ALL_UAVs, // non-VS stages (PS, CS, etc.)
 	RPK_RootParameterKeyCount,
 };
 
@@ -176,8 +177,14 @@ public:
 
 	inline uint32 UAVRDTBindSlot(EShaderFrequency ShaderStage) const
 	{
-		check(ShaderStage == SF_Pixel || ShaderStage == SF_Compute);
-		return BindSlotMap[ALL_UAVs];
+		check(ShaderStage == SF_Pixel || ShaderStage == SF_Vertex || ShaderStage == SF_Compute);
+		const uint32 MapSlotIndex = ShaderStage == SF_Vertex ? VS_UAVs : ALL_UAVs;
+		return BindSlotMap[MapSlotIndex];
+	}
+
+	inline static bool IsValidBindSlot(uint32 BindSlotMapIndex)
+	{
+		return BindSlotMapIndex < InvalidBindSlotMapIndex;
 	}
 
 	inline bool HasUAVs() const { return bHasUAVs; }
@@ -239,7 +246,7 @@ private:
 			return;
 		}
 
-		check(*pBindSlot == 0xFF);
+		check(*pBindSlot == InvalidBindSlotMapIndex);
 		*pBindSlot = RootParameterIndex;
 
 		bHasSamplers = true;
@@ -263,7 +270,7 @@ private:
 			return;
 		}
 
-		check(*pBindSlot == 0xFF);
+		check(*pBindSlot == InvalidBindSlotMapIndex);
 		*pBindSlot = RootParameterIndex;
 
 		bHasSRVs = true;
@@ -287,7 +294,7 @@ private:
 			return;
 		}
 
-		check(*pBindSlot == 0xFF);
+		check(*pBindSlot == InvalidBindSlotMapIndex);
 		*pBindSlot = RootParameterIndex;
 
 		bHasCBVs = true;
@@ -311,7 +318,7 @@ private:
 			return;
 		}
 
-		check(*pBindSlot == 0xFF);
+		check(*pBindSlot == InvalidBindSlotMapIndex);
 		*pBindSlot = RootParameterIndex;
 
 		bHasRootCBs = true;
@@ -319,11 +326,12 @@ private:
 
 	inline void SetUAVRDTBindSlot(EShaderFrequency SF, uint8 RootParameterIndex)
 	{
-		check(SF == SF_Pixel || SF == SF_Compute || SF == SF_NumFrequencies);
+		check(SF == SF_Pixel || SF == SF_Vertex || SF == SF_Compute || SF == SF_NumFrequencies);
 
-		uint8* pBindSlot = &BindSlotMap[ALL_UAVs];
+		const uint32 MapSlotIndex = SF == SF_Vertex ? VS_UAVs : ALL_UAVs;
+		uint8* pBindSlot = &BindSlotMap[MapSlotIndex];
 
-		check(*pBindSlot == 0xFF);
+		check(*pBindSlot == InvalidBindSlotMapIndex);
 		*pBindSlot = RootParameterIndex;
 
 		bHasUAVs = true;
@@ -445,6 +453,7 @@ private:
 
 	TRefCountPtr<ID3D12RootSignature> RootSignature;
 	uint8 BindSlotMap[RPK_RootParameterKeyCount];	// This map uses an enum as a key to lookup the root parameter index
+	static constexpr uint8 InvalidBindSlotMapIndex = 0xFF;
 	ShaderStage Stage[SF_NumFrequencies];
 	TRefCountPtr<ID3DBlob> RootSignatureBlob;
 
diff --git a/Engine/Source/Runtime/D3D12RHI/Private/D3D12StateCache.cpp b/Engine/Source/Runtime/D3D12RHI/Private/D3D12StateCache.cpp
--- a/Engine/Source/Runtime/D3D12RHI/Private/D3D12StateCache.cpp
+++ b/Engine/Source/Runtime/D3D12RHI/Private/D3D12StateCache.cpp
@@ -619,7 +619,7 @@ void FD3D12StateCache::ApplyResources(const FD3D12RootSignature* const pRootSign
 #if D3D12RHI_USE_CONSTANT_BUFFER_VIEWS
 	CBVSlotMask CurrentShaderDirtyCBVSlots[SF_NumStandardFrequencies] = {};
 #endif
-	UAVSlotMask CurrentShaderDirtyUAVSlots = 0;
+	UAVSlotMask CurrentShaderDirtyUAVSlots[SF_NumStandardFrequencies] = {};
 	uint32 NumUAVs = 0;
 	uint32 NumSRVs[SF_NumStandardFrequencies] = {};
 #if D3D12RHI_USE_CONSTANT_BUFFER_VIEWS
@@ -627,28 +627,34 @@ void FD3D12StateCache::ApplyResources(const FD3D12RootSignature* const pRootSign
 #endif
 	uint32 NumViews = 0;
 
-	const EShaderFrequency UAVStage = StartStage == SF_Compute ? SF_Compute : SF_Pixel;
+	EShaderFrequency ComputeUAVStages[] = {SF_Compute};
+	EShaderFrequency GraphicsUAVStages[] = {SF_Vertex, SF_Pixel};
+
+	TArrayView<EShaderFrequency> UAVStages = StartStage == SF_Compute ? MakeArrayView(ComputeUAVStages) : MakeArrayView(GraphicsUAVStages);
 
 	for (uint32 iTries = 0; iTries < 2; ++iTries)
 	{
 		if (bUAVs)
 		{
-			const UAVSlotMask CurrentShaderUAVRegisterMask = BitMask<UAVSlotMask>(PipelineState.Common.CurrentShaderUAVCounts[UAVStage]);
-			CurrentShaderDirtyUAVSlots = CurrentShaderUAVRegisterMask & PipelineState.Common.UAVCache.DirtySlotMask[UAVStage];
-			if (CurrentShaderDirtyUAVSlots)
+			for (EShaderFrequency UAVStage : UAVStages)
 			{
-				if (ResourceBindingTier <= D3D12_RESOURCE_BINDING_TIER_2)
-				{
-					// Tier 1 and 2 HW requires the full number of UAV descriptors defined in the root signature's descriptor table.
-					NumUAVs = pRootSignature->MaxUAVCount(UAVStage);
-				}
-				else
+				const UAVSlotMask CurrentShaderUAVRegisterMask = BitMask<UAVSlotMask>(PipelineState.Common.CurrentShaderUAVCounts[UAVStage]);
+				CurrentShaderDirtyUAVSlots[UAVStage] = CurrentShaderUAVRegisterMask & PipelineState.Common.UAVCache.DirtySlotMask[UAVStage];
+				if (CurrentShaderDirtyUAVSlots[UAVStage])
 				{
-					NumUAVs = PipelineState.Common.CurrentShaderUAVCounts[UAVStage];
-				}
+					if (ResourceBindingTier <= D3D12_RESOURCE_BINDING_TIER_2)
+					{
+						// Tier 1 and 2 HW requires the full number of UAV descriptors defined in the root signature's descriptor table.
+						NumUAVs = pRootSignature->MaxUAVCount(UAVStage);
+					}
+					else
+					{
+						NumUAVs = PipelineState.Common.CurrentShaderUAVCounts[UAVStage];
+					}
 
-				check(NumUAVs > 0 && NumUAVs <= MAX_UAVS);
-				NumViews += NumUAVs;
+					check(NumUAVs > 0 && NumUAVs <= MAX_UAVS);
+					NumViews += NumUAVs;
+				}
 			}
 		}
 
@@ -725,11 +731,16 @@ void FD3D12StateCache::ApplyResources(const FD3D12RootSignature* const pRootSign
 	uint32 ViewHeapSlot = DescriptorCache.GetCurrentViewHeap()->ReserveSlots(NumViews);
 
 	// Unordered access views
-	if (CurrentShaderDirtyUAVSlots)
 	{
 		SCOPE_CYCLE_COUNTER(STAT_D3D12ApplyStateSetUAVTime);
-		const D3D12_GPU_DESCRIPTOR_HANDLE BindDescriptor = DescriptorCache.BuildUAVTable(UAVStage, pRootSignature, PipelineState.Common.UAVCache, CurrentShaderDirtyUAVSlots, NumUAVs, ViewHeapSlot);
-		DescriptorCache.SetUAVTable(UAVStage, pRootSignature, PipelineState.Common.UAVCache, NumUAVs, BindDescriptor);
+		for (EShaderFrequency UAVStage : UAVStages)
+		{
+			if (CurrentShaderDirtyUAVSlots[UAVStage])
+			{
+				const D3D12_GPU_DESCRIPTOR_HANDLE BindDescriptor = DescriptorCache.BuildUAVTable(UAVStage, pRootSignature, PipelineState.Common.UAVCache, CurrentShaderDirtyUAVSlots[UAVStage], NumUAVs, ViewHeapSlot);
+				DescriptorCache.SetUAVTable(UAVStage, pRootSignature, PipelineState.Common.UAVCache, NumUAVs, BindDescriptor);
+			}
+		}
 	}
 
 	// Shader resource views
diff --git a/Engine/Source/Runtime/D3D12RHI/Private/D3D12Util.cpp b/Engine/Source/Runtime/D3D12RHI/Private/D3D12Util.cpp
--- a/Engine/Source/Runtime/D3D12RHI/Private/D3D12Util.cpp
+++ b/Engine/Source/Runtime/D3D12RHI/Private/D3D12Util.cpp
@@ -1318,7 +1318,7 @@ const FD3D12RootSignature* FD3D12Adapter::GetRootSignature(const FBoundShaderSta
 
 	const D3D12_RESOURCE_BINDING_TIER ResourceBindingTier = GetResourceBindingTier();
 
-	QuantizeBoundShaderStateCommon(QBSS, FD3D12DynamicRHI::ResourceCast(BSS.GetVertexShader()),        ResourceBindingTier, SV_Vertex);
+	QuantizeBoundShaderStateCommon(QBSS, FD3D12DynamicRHI::ResourceCast(BSS.GetVertexShader()),        ResourceBindingTier, SV_Vertex, true /*bAllowUAVs*/);
 #if PLATFORM_SUPPORTS_MESH_SHADERS
 	QuantizeBoundShaderStateCommon(QBSS, FD3D12DynamicRHI::ResourceCast(BSS.GetMeshShader()),          ResourceBindingTier, SV_Mesh);
 	QuantizeBoundShaderStateCommon(QBSS, FD3D12DynamicRHI::ResourceCast(BSS.GetAmplificationShader()), ResourceBindingTier, SV_Amplification);
diff --git a/Engine/Source/Runtime/D3D12RHI/Public/D3D12RootSignatureDefinitions.h b/Engine/Source/Runtime/D3D12RHI/Public/D3D12RootSignatureDefinitions.h
--- a/Engine/Source/Runtime/D3D12RHI/Public/D3D12RootSignatureDefinitions.h
+++ b/Engine/Source/Runtime/D3D12RHI/Public/D3D12RootSignatureDefinitions.h
@@ -414,7 +414,8 @@ namespace D3D12ShaderUtils
 		}
 #endif
 
-		Creator.AddTable(ERootSignatureVisibility::All, ERootSignatureRangeType::UAV, MAX_UAVS);
+		Creator.AddTable(ERootSignatureVisibility::Pixel, ERootSignatureRangeType::UAV, MAX_UAVS);
+		Creator.AddTable(ERootSignatureVisibility::Vertex, ERootSignatureRangeType::UAV, MAX_UAVS);
 
 		if (EnumHasAnyFlags(InFlags, ED3D12RootSignatureFlags::RootConstants))
 		{
diff --git a/Engine/Source/Runtime/Engine/Private/ShaderCompiler/ShaderCompiler.cpp b/Engine/Source/Runtime/Engine/Private/ShaderCompiler/ShaderCompiler.cpp
--- a/Engine/Source/Runtime/Engine/Private/ShaderCompiler/ShaderCompiler.cpp
+++ b/Engine/Source/Runtime/Engine/Private/ShaderCompiler/ShaderCompiler.cpp
@@ -8164,6 +8164,11 @@ void GlobalBeginCompileShader(
 		SET_SHADER_DEFINE(Input.Environment, PLATFORM_SUPPORTS_VERTEX_SHADER_SRVS, bSupportsVertexShaderSRVs ? 1u : 0u);
 	}
 
+	{
+		const bool bSupportsVertexShaderUAVs = FDataDrivenShaderPlatformInfo::GetSupportsVertexShaderUAVs((EShaderPlatform)Target.Platform) != ERHIFeatureSupport::Unsupported;
+		SET_SHADER_DEFINE(Input.Environment, PLATFORM_SUPPORTS_VERTEX_SHADER_UAVS, bSupportsVertexShaderUAVs ? 1u : 0u);
+	}
+
 	{
 		const uint32 MaxSamplers = FDataDrivenShaderPlatformInfo::GetMaxSamplers((EShaderPlatform)Target.Platform);
 		SET_SHADER_DEFINE(Input.Environment, PLATFORM_MAX_SAMPLERS, MaxSamplers);
diff --git a/Engine/Source/Runtime/RHI/Private/DataDrivenShaderPlatformInfo.cpp b/Engine/Source/Runtime/RHI/Private/DataDrivenShaderPlatformInfo.cpp
--- a/Engine/Source/Runtime/RHI/Private/DataDrivenShaderPlatformInfo.cpp
+++ b/Engine/Source/Runtime/RHI/Private/DataDrivenShaderPlatformInfo.cpp
@@ -145,6 +145,7 @@ void FGenericDataDrivenShaderPlatformInfo::SetDefaultValues()
 	bSupportsWaterIndirectDraw = true;
 	bSupportsAsyncPipelineCompilation = true;
 	bSupportsVertexShaderSRVs = true; // Explicitly overriden to false for ES 3.1 platforms via DDPI ini
+	bSupportsVertexShaderUAVs = uint32(ERHIFeatureSupport::Unsupported);
 	bSupportsManualVertexFetch = true;
 	bSupportsVolumeTextureAtomics = true;
 	bSupportsClipDistance = true;
@@ -277,6 +278,7 @@ void FGenericDataDrivenShaderPlatformInfo::ParseDataDrivenShaderInfo(const FConf
 	GET_SECTION_BOOL_HELPER(bSupportsWaterIndirectDraw);
 	GET_SECTION_BOOL_HELPER(bSupportsAsyncPipelineCompilation);
 	GET_SECTION_BOOL_HELPER(bSupportsVertexShaderSRVs);
+	GET_SECTION_SUPPORT_HELPER(bSupportsVertexShaderUAVs);
 	GET_SECTION_BOOL_HELPER(bSupportsManualVertexFetch);
 	GET_SECTION_BOOL_HELPER(bRequiresReverseCullingOnMobile);
 	GET_SECTION_BOOL_HELPER(bOverrideFMaterial_NeedsGBufferEnabled);
@@ -477,6 +479,7 @@ void FGenericDataDrivenShaderPlatformInfo::UpdatePreviewPlatforms()
 				PREVIEW_USE_RUNTIME_VALUE(MaxMeshShaderThreadGroupSize);
 				PREVIEW_USE_RUNTIME_VALUE(bSupportsSceneDataCompressedTransforms);
 				PREVIEW_USE_RUNTIME_VALUE(bSupportsVertexShaderSRVs);
+				PREVIEW_USE_RUNTIME_VALUE(bSupportsVertexShaderUAVs);
 				PREVIEW_USE_RUNTIME_VALUE(bSupportsManualVertexFetch);
 				PREVIEW_USE_RUNTIME_VALUE(bSupportsRealTypes);
 				PREVIEW_USE_RUNTIME_VALUE(bSupportsUniformBufferObjects);
diff --git a/Engine/Source/Runtime/RHI/Public/DataDrivenShaderPlatformInfo.h b/Engine/Source/Runtime/RHI/Public/DataDrivenShaderPlatformInfo.h
--- a/Engine/Source/Runtime/RHI/Public/DataDrivenShaderPlatformInfo.h
+++ b/Engine/Source/Runtime/RHI/Public/DataDrivenShaderPlatformInfo.h
@@ -99,6 +99,7 @@ class FGenericDataDrivenShaderPlatformInfo
 	uint32 bSupportsWaterIndirectDraw : 1;
 	uint32 bSupportsAsyncPipelineCompilation : 1;
 	uint32 bSupportsVertexShaderSRVs : 1; // Whether SRVs can be bound to vertex shaders (may be independent from ManualVertexFetch)
+	uint32 bSupportsVertexShaderUAVs : int32(ERHIFeatureSupport::NumBits); // Whether UAVs can be bound to vertex shaders. Requires run-time check of GRHIGlobals.SupportsVertexShaderUAVs.
 	uint32 bSupportsManualVertexFetch : 1;
 	uint32 bRequiresReverseCullingOnMobile : 1;
 	uint32 bOverrideFMaterial_NeedsGBufferEnabled : 1;
@@ -677,6 +678,12 @@ public:
 		return Infos[Platform].bSupportsVertexShaderSRVs;
 	}
 
+	static FORCEINLINE_DEBUGGABLE const ERHIFeatureSupport GetSupportsVertexShaderUAVs(const FStaticShaderPlatform Platform)
+	{
+		check(IsValid(Platform));
+		return ERHIFeatureSupport(Infos[Platform].bSupportsVertexShaderUAVs);
+	}
+
 	static FORCEINLINE_DEBUGGABLE const bool GetSupportsManualVertexFetch(const FStaticShaderPlatform Platform)
 	{
 		check(IsValid(Platform));
diff --git a/Engine/Source/Runtime/RHI/Public/RHIGlobals.h b/Engine/Source/Runtime/RHI/Public/RHIGlobals.h
--- a/Engine/Source/Runtime/RHI/Public/RHIGlobals.h
+++ b/Engine/Source/Runtime/RHI/Public/RHIGlobals.h
@@ -140,6 +140,9 @@ struct FRHIGlobals
 	// true if the RHI supports Pixel Shader UAV
 	bool SupportsPixelShaderUAVs = true;
 
+	// true if the RHI supports Vertex Shader UAV
+	bool SupportsVertexShaderUAVs = false;
+
 	/** true if PF_G8 render targets are supported */
 	TRHIGlobal<bool> SupportsRenderTargetFormat_PF_G8 = true;
 
diff --git a/Engine/Source/Runtime/Windows/D3D11RHI/Private/D3D11Commands.cpp b/Engine/Source/Runtime/Windows/D3D11RHI/Private/D3D11Commands.cpp
--- a/Engine/Source/Runtime/Windows/D3D11RHI/Private/D3D11Commands.cpp
+++ b/Engine/Source/Runtime/Windows/D3D11RHI/Private/D3D11Commands.cpp
@@ -353,13 +353,13 @@ struct FD3D11ResourceBinder
 		{
 			RHI.InternalSetUAVCS(Index, FD3D11DynamicRHI::ResourceCast(InUnorderedAccessView));
 		}
-		else if (ShaderFrequency == SF_Pixel)
+		else if (ShaderFrequency == SF_Pixel || ShaderFrequency == SF_Vertex)
 		{
-			RHI.InternalSetUAVPS(Index, FD3D11DynamicRHI::ResourceCast(InUnorderedAccessView));
+			RHI.InternalSetUAVVSPS(Index, FD3D11DynamicRHI::ResourceCast(InUnorderedAccessView));
 		}
 		else
 		{
-			checkf(false, TEXT("UAVs are not supported on vertex and geometry shaders."));
+			checkf(false, TEXT("UAVs are only supported in compute, pixel and vertex shaders."));
 		}
 	}
 
@@ -630,7 +630,7 @@ void FD3D11DynamicRHI::InternalSetUAVCS(uint32 BindIndex, FD3D11UnorderedAccessV
 	Direct3DDeviceIMContext->CSSetUnorderedAccessViews(BindIndex, 1, &D3D11UAV, &InitialCount);
 }
 
-void FD3D11DynamicRHI::InternalSetUAVPS(uint32 BindIndex, FD3D11UnorderedAccessView* UnorderedAccessViewRHI)
+void FD3D11DynamicRHI::InternalSetUAVVSPS(uint32 BindIndex, FD3D11UnorderedAccessView* UnorderedAccessViewRHI)
 {
 	check(BindIndex < D3D11_PS_CS_UAV_REGISTER_COUNT);
 	if (CurrentUAVs[BindIndex] != UnorderedAccessViewRHI)
@@ -1032,47 +1032,52 @@ void FD3D11DynamicRHI::CommitGraphicsResourceTables()
 	FD3D11BoundShaderState* RESTRICT CurrentBoundShaderState = (FD3D11BoundShaderState*)BoundShaderStateHistory.GetLast();
 	check(CurrentBoundShaderState);
 
-	auto* PixelShader = CurrentBoundShaderState->GetPixelShader();
-	if (PixelShader)
-	{
-		// Because d3d11 binding uses the same slots for UAVs and RTVs, we have to rebind when two shaders with different sets of rendertargets are bound,
-		// as they can potentially be used by UAVs, which can cause them to unbind RTVs used by subsequent shaders.
-		bool bRTVInvalidate = false;
-		uint32 UAVMask = PixelShader->UAVMask & CurrentRTVOverlapMask;
-		if (GDX11ReduceRTVRebinds && 
-			(0 != ((~CurrentUAVMask) & UAVMask) && CurrentUAVMask == (CurrentUAVMask & UAVMask)))
-		{
-			//if the mask only -adds- uav binds, no RTs will be missing so we just grow the mask
-			CurrentUAVMask = UAVMask;
-		}
-		else if (CurrentUAVMask != UAVMask)
-		{
-			bRTVInvalidate = true;
-			CurrentUAVMask = UAVMask;
-		}
+	bool bRTVInvalidate = false;
+	uint32 UAVMask = 0;
 
-		if (bRTVInvalidate)
-		{
-			CommitRenderTargets(true);
-			DirtyUniformBuffers[SF_Pixel] = -1;
-		}
-
-		SetResourcesFromTables(PixelShader);
-
-		if (UAVSChanged)
-		{
-			CommitUAVs();
-		}
+	if (auto* Shader = CurrentBoundShaderState->GetPixelShader())
+	{
+		UAVMask |= Shader->UAVMask & CurrentRTVOverlapMask;
+		SetResourcesFromTables(Shader);
 	}
 
 	if (auto* Shader = CurrentBoundShaderState->GetVertexShader())
 	{
+		UAVMask |= Shader->UAVMask & CurrentRTVOverlapMask;
 		SetResourcesFromTables(Shader);
 	}
 	if (auto* Shader = CurrentBoundShaderState->GetGeometryShader())
 	{
+		UAVMask |= Shader->UAVMask & CurrentRTVOverlapMask;
 		SetResourcesFromTables(Shader);
 	}
+
+	// Because d3d11 binding uses the same slots for UAVs and RTVs, we have to rebind when two shaders with different sets of rendertargets are bound,
+	// as they can potentially be used by UAVs, which can cause them to unbind RTVs used by subsequent shaders.
+	if (GDX11ReduceRTVRebinds &&
+		(0 != ((~CurrentUAVMask) & UAVMask) && CurrentUAVMask == (CurrentUAVMask & UAVMask)))
+	{
+		//if the mask only -adds- uav binds, no RTs will be missing so we just grow the mask
+		CurrentUAVMask = UAVMask;
+	}
+	else if (CurrentUAVMask != UAVMask)
+	{
+		bRTVInvalidate = true;
+		CurrentUAVMask = UAVMask;
+	}
+
+	if (bRTVInvalidate)
+	{
+		CommitRenderTargets(true);
+		DirtyUniformBuffers[SF_Pixel] = -1;
+		DirtyUniformBuffers[SF_Vertex] = -1;
+		DirtyUniformBuffers[SF_Geometry] = -1;
+	}
+
+	if (UAVSChanged)
+	{
+		CommitUAVs();
+	}
 }
 
 void FD3D11DynamicRHI::CommitComputeResourceTables(FD3D11ComputeShader* InComputeShader)
diff --git a/Engine/Source/Runtime/Windows/D3D11RHI/Private/D3D11Device.cpp b/Engine/Source/Runtime/Windows/D3D11RHI/Private/D3D11Device.cpp
--- a/Engine/Source/Runtime/Windows/D3D11RHI/Private/D3D11Device.cpp
+++ b/Engine/Source/Runtime/Windows/D3D11RHI/Private/D3D11Device.cpp
@@ -204,6 +204,10 @@ FD3D11DynamicRHI::FD3D11DynamicRHI(IDXGIFactory1* InDXGIFactory1, D3D_FEATURE_LE
 	GRHIMaxDispatchThreadGroupsPerDimension.Y = D3D11_CS_DISPATCH_MAX_THREAD_GROUPS_PER_DIMENSION;
 	GRHIMaxDispatchThreadGroupsPerDimension.Z = D3D11_CS_DISPATCH_MAX_THREAD_GROUPS_PER_DIMENSION;
 
+	// All D3D11.1 hardware on Windows 8+ supports binding UAVs to Vertex Shaders.
+	// Enable run-time support if corresponding bit is set in DDSPI.
+	GRHIGlobals.SupportsVertexShaderUAVs = true;
+
 	GRHIGlobals.NeedsShaderUnbinds = true;
 
 	GRHITransitionPrivateData_SizeInBytes = sizeof(FD3D11TransitionData);
diff --git a/Engine/Source/Runtime/Windows/D3D11RHI/Private/D3D11RHIPrivate.h b/Engine/Source/Runtime/Windows/D3D11RHI/Private/D3D11RHIPrivate.h
--- a/Engine/Source/Runtime/Windows/D3D11RHI/Private/D3D11RHIPrivate.h
+++ b/Engine/Source/Runtime/Windows/D3D11RHI/Private/D3D11RHIPrivate.h
@@ -546,7 +546,7 @@ public:
 	virtual void RHISetBlendFactor(const FLinearColor& BlendFactor) final override;
 	void SetRenderTargets(uint32 NumSimultaneousRenderTargets, const FRHIRenderTargetView* NewRenderTargets, const FRHIDepthRenderTargetView* NewDepthStencilTarget);
 	void InternalSetUAVCS(uint32 BindIndex, FD3D11UnorderedAccessView* UnorderedAccessViewRHI);
-	void InternalSetUAVPS(uint32 BindIndex, FD3D11UnorderedAccessView* UnorderedAccessViewRHI);
+	void InternalSetUAVVSPS(uint32 BindIndex, FD3D11UnorderedAccessView* UnorderedAccessViewRHI);
 	void SetRenderTargetsAndClear(const FRHISetRenderTargetsInfo& RenderTargetsInfo);
 	virtual void RHIDrawPrimitive(uint32 BaseVertexIndex, uint32 NumPrimitives, uint32 NumInstances) final override;
 	virtual void RHIDrawPrimitiveIndirect(FRHIBuffer* ArgumentBuffer, uint32 ArgumentOffset) final override;
