Fix buffer descriptor type in Test_GraphicsUAV_VertexShader

diff --git a/Engine/Plugins/Tests/RHITests/Source/RHITests/Private/RHIGraphicsUAVTests.cpp b/Engine/Plugins/Tests/RHITests/Source/RHITests/Private/RHIGraphicsUAVTests.cpp
--- a/Engine/Plugins/Tests/RHITests/Source/RHITests/Private/RHIGraphicsUAVTests.cpp
+++ b/Engine/Plugins/Tests/RHITests/Source/RHITests/Private/RHIGraphicsUAVTests.cpp
@@ -173,14 +173,15 @@ bool FRHIGraphicsUAVTests::Test_GraphicsUAV_VertexShader(FRHICommandListImmediat
 	static constexpr uint32 OutputBufferSize = OutputBufferStride * MaxVertices;
 	FRHIResourceCreateInfo OutputBufferCreateInfo(TEXT("GraphicsUAVTests_VertexShaderOutput"));
 
+	// NOTE: using a structured buffer here as workaround for UE-212251
 	FBufferRHIRef OutputBuffer = RHICmdList.CreateBuffer(OutputBufferSize,
-		EBufferUsageFlags::UnorderedAccess | EBufferUsageFlags::SourceCopy,
+		EBufferUsageFlags::UnorderedAccess | EBufferUsageFlags::SourceCopy | EBufferUsageFlags::StructuredBuffer,
 		OutputBufferStride, ERHIAccess::UAVCompute, OutputBufferCreateInfo);
 
 	FUnorderedAccessViewRHIRef OutputBufferUAV = RHICmdList.CreateUnorderedAccessView(OutputBuffer,
 		FRHIViewDesc::CreateBufferUAV()
-		.SetType(FRHIViewDesc::EBufferType::Typed)
-		.SetFormat(PF_R32_UINT));
+		.SetType(FRHIViewDesc::EBufferType::Structured)
+		.SetStride(4));
 
 	RHICmdList.ClearUAVUint(OutputBufferUAV, FUintVector4(~0u));
 
