Fix errors during RHI draw tests on some platforms
 - Tests were using a vertex shader compiled with "CFLAG_IndirectDraw" in a standard direct draw, which is not allowed.

diff --git a/Engine/Plugins/Tests/RHITests/Source/RHITests/Private/RHIDrawTests.cpp b/Engine/Plugins/Tests/RHITests/Source/RHITests/Private/RHIDrawTests.cpp
--- a/Engine/Plugins/Tests/RHITests/Source/RHITests/Private/RHIDrawTests.cpp
+++ b/Engine/Plugins/Tests/RHITests/Source/RHITests/Private/RHIDrawTests.cpp
@@ -7,11 +7,22 @@
 #include "RHIStaticStates.h"
 #include "ShaderCompilerCore.h"
 
-class FTestDrawInstancedVS : public FGlobalShader
+class FTestDrawInstancedVS_Direct : public FGlobalShader
 {
-	DECLARE_GLOBAL_SHADER(FTestDrawInstancedVS);
-	FTestDrawInstancedVS() = default;
-	FTestDrawInstancedVS(const ShaderMetaType::CompiledShaderInitializerType& Initializer)
+	DECLARE_GLOBAL_SHADER(FTestDrawInstancedVS_Direct);
+
+	FTestDrawInstancedVS_Direct() = default;
+	FTestDrawInstancedVS_Direct(const ShaderMetaType::CompiledShaderInitializerType& Initializer)
+		: FGlobalShader(Initializer)
+	{}
+};
+
+class FTestDrawInstancedVS_Indirect : public FGlobalShader
+{
+	DECLARE_GLOBAL_SHADER(FTestDrawInstancedVS_Indirect);
+
+	FTestDrawInstancedVS_Indirect() = default;
+	FTestDrawInstancedVS_Indirect(const ShaderMetaType::CompiledShaderInitializerType& Initializer)
 		: FGlobalShader(Initializer)
 	{
 	}
@@ -20,7 +31,9 @@ class FTestDrawInstancedVS : public FGlobalShader
 		OutEnvironment.CompilerFlags.Add(CFLAG_IndirectDraw);
 	}
 };
-IMPLEMENT_GLOBAL_SHADER(FTestDrawInstancedVS, "/Plugin/RHITests/Private/TestDrawInstanced.usf", "TestDrawInstancedMainVS", SF_Vertex);
+
+IMPLEMENT_GLOBAL_SHADER(FTestDrawInstancedVS_Direct, "/Plugin/RHITests/Private/TestDrawInstanced.usf", "TestDrawInstancedMainVS", SF_Vertex);
+IMPLEMENT_GLOBAL_SHADER(FTestDrawInstancedVS_Indirect, "/Plugin/RHITests/Private/TestDrawInstanced.usf", "TestDrawInstancedMainVS", SF_Vertex);
 
 class FTestDrawInstancedPS : public FGlobalShader
 {
@@ -45,7 +58,8 @@ namespace
 struct FDrawTestResources
 {
 	FDrawTestResources(FRHICommandListImmediate& RHICmdList)
-		: VertexShader(GetGlobalShaderMap(GMaxRHIFeatureLevel))
+		: VertexShader_Direct(GetGlobalShaderMap(GMaxRHIFeatureLevel))
+		, VertexShader_Indirect(GetGlobalShaderMap(GMaxRHIFeatureLevel))
 		, PixelShader(GetGlobalShaderMap(GMaxRHIFeatureLevel))
 	{		
 		FRHITextureDesc RenderTargetTextureDesc(ETextureDimension::Texture2D, ETextureCreateFlags::RenderTargetable, PF_B8G8R8A8, FClearValueBinding(), RenderTargetSize, 1, 1, 1, 1, 0);
@@ -60,13 +74,16 @@ struct FDrawTestResources
 
 		VertexDeclarationRHI = PipelineStateCache::GetOrCreateVertexDeclaration(VertexDeclarationElements);
 
-		GraphicsPSOInit.BoundShaderState.VertexShaderRHI = VertexShader.GetVertexShader();
-		GraphicsPSOInit.BoundShaderState.VertexDeclarationRHI = VertexDeclarationRHI;
-		GraphicsPSOInit.BoundShaderState.PixelShaderRHI = PixelShader.GetPixelShader();
-		GraphicsPSOInit.DepthStencilState = TStaticDepthStencilState<false, CF_Always>::GetRHI();
-		GraphicsPSOInit.BlendState = TStaticBlendState<>::GetRHI();
-		GraphicsPSOInit.RasterizerState = TStaticRasterizerState<>::GetRHI();
-		GraphicsPSOInit.PrimitiveType = EPrimitiveType::PT_TriangleList;
+		GraphicsPSOInit_Direct.BoundShaderState.VertexShaderRHI = VertexShader_Direct.GetVertexShader();
+		GraphicsPSOInit_Direct.BoundShaderState.VertexDeclarationRHI = VertexDeclarationRHI;
+		GraphicsPSOInit_Direct.BoundShaderState.PixelShaderRHI = PixelShader.GetPixelShader();
+		GraphicsPSOInit_Direct.DepthStencilState = TStaticDepthStencilState<false, CF_Always>::GetRHI();
+		GraphicsPSOInit_Direct.BlendState = TStaticBlendState<>::GetRHI();
+		GraphicsPSOInit_Direct.RasterizerState = TStaticRasterizerState<>::GetRHI();
+		GraphicsPSOInit_Direct.PrimitiveType = EPrimitiveType::PT_TriangleList;
+
+		GraphicsPSOInit_Indirect = GraphicsPSOInit_Direct;
+		GraphicsPSOInit_Indirect.BoundShaderState.VertexShaderRHI = VertexShader_Indirect.GetVertexShader();
 
 		// D3D12 does not have a way to get the base instance ID (SV_InstanceID always starts from 0), so we must emulate it...
 		uint32 InstanceIDs[MaxInstances];
@@ -118,9 +135,13 @@ struct FDrawTestResources
 	static constexpr uint32 MaxInstances = 8;
 
 	FVertexDeclarationRHIRef VertexDeclarationRHI;
-	TShaderMapRef<FTestDrawInstancedVS> VertexShader;
+
+	TShaderMapRef<FTestDrawInstancedVS_Direct> VertexShader_Direct;
+	TShaderMapRef<FTestDrawInstancedVS_Indirect> VertexShader_Indirect;
 	TShaderMapRef<FTestDrawInstancedPS> PixelShader;
-	FGraphicsPipelineStateInitializer GraphicsPSOInit;
+
+	FGraphicsPipelineStateInitializer GraphicsPSOInit_Direct;
+	FGraphicsPipelineStateInitializer GraphicsPSOInit_Indirect;
 
 	FIntPoint RenderTargetSize = FIntPoint(4, 4);
 	FTextureRHIRef RenderTarget;
@@ -177,8 +198,12 @@ bool FRHIDrawTests::InternalDrawBaseVertexAndInstance(FRHICommandListImmediate&
 	RHICmdList.BeginRenderPass(RenderPassInfo, TestName);
 	RHICmdList.SetViewport(0, 0, 0, float(Resources.RenderTargetSize.X), float(Resources.RenderTargetSize.Y), 1);
 
-	RHICmdList.ApplyCachedRenderTargets(Resources.GraphicsPSOInit);
-	SetGraphicsPipelineState(RHICmdList, Resources.GraphicsPSOInit, 0);
+	FGraphicsPipelineStateInitializer& PSOInit = DrawKind == EDrawKind::Direct
+		? Resources.GraphicsPSOInit_Direct
+		: Resources.GraphicsPSOInit_Indirect;
+
+	RHICmdList.ApplyCachedRenderTargets(PSOInit);
+	SetGraphicsPipelineState(RHICmdList, PSOInit, 0);
 
 	check(Resources.InstanceIDBuffer->GetStride() == 4);
 	RHICmdList.SetStreamSource(0, Resources.VertexBuffer, 0);
@@ -286,8 +311,8 @@ bool FRHIDrawTests::Test_MultiDrawIndirect(FRHICommandListImmediate& RHICmdList)
 	RHICmdList.BeginRenderPass(RenderPassInfo, TEXT("Test_MultiDrawIndirect"));
 	RHICmdList.SetViewport(0, 0, 0, float(Resources.RenderTargetSize.X), float(Resources.RenderTargetSize.Y), 1);
 
-	RHICmdList.ApplyCachedRenderTargets(Resources.GraphicsPSOInit);
-	SetGraphicsPipelineState(RHICmdList, Resources.GraphicsPSOInit, 0);
+	RHICmdList.ApplyCachedRenderTargets(Resources.GraphicsPSOInit_Indirect);
+	SetGraphicsPipelineState(RHICmdList, Resources.GraphicsPSOInit_Indirect, 0);
 
 	check(Resources.InstanceIDBuffer->GetStride() == 4);
 	RHICmdList.SetStreamSource(0, Resources.VertexBuffer, 0);
