More upload helpers
- Adding FRHIResourceCreateInfoUploadArray to the new RHIResourceUtils.h
- FRHIResourceCreateInfoUploadArray is a FRHIResourceCreateInfo type that manages its own FResourceArrayUploadArrayView, meaning it can be used to replace the multi-line code it previously took to upload data from a container type.
- Adding UE::RHIResourceUtils::CreateBufferWithData to further reduce the code it takes to create a buffer with initial data.

diff --git a/Engine/Plugins/Tests/RHITests/Source/RHITests/Private/RHIDrawTests.cpp b/Engine/Plugins/Tests/RHITests/Source/RHITests/Private/RHIDrawTests.cpp
--- a/Engine/Plugins/Tests/RHITests/Source/RHITests/Private/RHIDrawTests.cpp
+++ b/Engine/Plugins/Tests/RHITests/Source/RHITests/Private/RHIDrawTests.cpp
@@ -91,7 +91,7 @@ struct FDrawTestResources
 		{
 			InstanceIDs[i] = i;
 		}
-		InstanceIDBuffer = CreateBufferWithData(EBufferUsageFlags::VertexBuffer, ERHIAccess::VertexOrIndexBuffer, TEXT("DrawTest_InstanceID"), MakeArrayView(InstanceIDs));
+		InstanceIDBuffer = UE::RHIResourceUtils::CreateBufferWithData(RHICmdList, EBufferUsageFlags::VertexBuffer, ERHIAccess::VertexOrIndexBuffer, TEXT("DrawTest_InstanceID"), MakeArrayView(InstanceIDs));
 
 		// Indices for 4 triangles
 		const uint16 Indices[NumTotalVertices] =
@@ -102,7 +102,7 @@ struct FDrawTestResources
 			9, 10, 11 // degenerate
 		};
 
-		IndexBuffer = CreateBufferWithData(EBufferUsageFlags::IndexBuffer, ERHIAccess::VertexOrIndexBuffer, TEXT("DrawTest_IndexBuffer"), MakeArrayView(Indices));
+		IndexBuffer = UE::RHIResourceUtils::CreateBufferWithData(RHICmdList, EBufferUsageFlags::IndexBuffer, ERHIAccess::VertexOrIndexBuffer, TEXT("DrawTest_IndexBuffer"), MakeArrayView(Indices));
 
 		TArray<FVector4f> Vertices;
 		Vertices.Reserve(NumTotalVertices);
@@ -119,7 +119,7 @@ struct FDrawTestResources
 		{
 			Vertices.Add(FVector4f(0.0f, 0.0f, 0.0f, 1.0f));
 		}
-		VertexBuffer = CreateBufferWithData(EBufferUsageFlags::VertexBuffer, ERHIAccess::VertexOrIndexBuffer, TEXT("DrawTest_VertexBuffer"), MakeArrayView(Vertices));
+		VertexBuffer = UE::RHIResourceUtils::CreateBufferWithData(RHICmdList, EBufferUsageFlags::VertexBuffer, ERHIAccess::VertexOrIndexBuffer, TEXT("DrawTest_VertexBuffer"), MakeArrayView(Vertices));
 
 		static constexpr uint32 OutputBufferStride = sizeof(uint32);
 		static constexpr uint32 OutputBufferSize = OutputBufferStride * MaxInstances;
@@ -188,7 +188,7 @@ bool FRHIDrawTests::InternalDrawBaseVertexAndInstance(FRHICommandListImmediate&
 	FBufferRHIRef DrawArgBuffer;
 	if (DrawKind == EDrawKind::Indirect)
 	{
-		DrawArgBuffer = CreateBufferWithData(EBufferUsageFlags::DrawIndirect | EBufferUsageFlags::UnorderedAccess | EBufferUsageFlags::VertexBuffer, ERHIAccess::IndirectArgs,
+		DrawArgBuffer = UE::RHIResourceUtils::CreateBufferWithData(RHICmdList, EBufferUsageFlags::DrawIndirect | EBufferUsageFlags::UnorderedAccess | EBufferUsageFlags::VertexBuffer, ERHIAccess::IndirectArgs,
 			TEXT("InternalDrawBaseVertexAndInstance_DrawArgs"), MakeArrayView(DrawArgs));
 	}
 
@@ -281,7 +281,7 @@ bool FRHIDrawTests::Test_MultiDrawIndirect(FRHICommandListImmediate& RHICmdList)
 	FDrawTestResources Resources(RHICmdList);
 
 	const uint32 CountValues[4] = { 1, 1, 16, 0 };
-	FBufferRHIRef CountBuffer = CreateBufferWithData(EBufferUsageFlags::DrawIndirect | EBufferUsageFlags::UnorderedAccess, ERHIAccess::IndirectArgs, TEXT("Test_MultiDrawIndirect_Count"), MakeArrayView(CountValues));
+	FBufferRHIRef CountBuffer = UE::RHIResourceUtils::CreateBufferWithData(RHICmdList, EBufferUsageFlags::DrawIndirect | EBufferUsageFlags::UnorderedAccess, ERHIAccess::IndirectArgs, TEXT("Test_MultiDrawIndirect_Count"), MakeArrayView(CountValues));
 
 	const FRHIDrawIndexedIndirectParameters DrawArgs[] =
 	{
@@ -297,7 +297,7 @@ bool FRHIDrawTests::Test_MultiDrawIndirect(FRHICommandListImmediate& RHICmdList)
 
 	const uint32 ExpectedDrawnInstances[Resources.MaxInstances] = { 1, 0, 1, 1, 0, 1, 1, 0 };
 
-	FBufferRHIRef DrawArgBuffer = CreateBufferWithData(EBufferUsageFlags::DrawIndirect | EBufferUsageFlags::UnorderedAccess | EBufferUsageFlags::VertexBuffer, ERHIAccess::IndirectArgs,
+	FBufferRHIRef DrawArgBuffer = UE::RHIResourceUtils::CreateBufferWithData(RHICmdList, EBufferUsageFlags::DrawIndirect | EBufferUsageFlags::UnorderedAccess | EBufferUsageFlags::VertexBuffer, ERHIAccess::IndirectArgs,
 		TEXT("Test_MultiDrawIndirect_DrawArgs"), MakeArrayView(DrawArgs));
 
 	RHICmdList.ClearUAVUint(Resources.OutputBufferUAV, FUintVector4(0));
diff --git a/Engine/Plugins/Tests/RHITests/Source/RHITests/Private/RHIGraphicsUAVTests.cpp b/Engine/Plugins/Tests/RHITests/Source/RHITests/Private/RHIGraphicsUAVTests.cpp
--- a/Engine/Plugins/Tests/RHITests/Source/RHITests/Private/RHIGraphicsUAVTests.cpp
+++ b/Engine/Plugins/Tests/RHITests/Source/RHITests/Private/RHIGraphicsUAVTests.cpp
@@ -89,7 +89,7 @@ static void Test_GraphicsUAV_Common(FRHICommandListImmediate& RHICmdList,
 	Vertices.Add(FVector4f(-1.0f, +3.0f, 0.0f, 1.0f));
 	Vertices.Add(FVector4f(+3.0f, -1.0f, 0.0f, 1.0f));
 
-	FBufferRHIRef VertexBuffer = CreateBufferWithData(EBufferUsageFlags::VertexBuffer,
+	FBufferRHIRef VertexBuffer = UE::RHIResourceUtils::CreateBufferWithData(RHICmdList, EBufferUsageFlags::VertexBuffer,
 		ERHIAccess::VertexOrIndexBuffer, TEXT("GraphicsUAVTests_VertexBuffer"), MakeArrayView(Vertices));
 
 	FRHITexture* ColorRTs[1] = { RenderTarget.GetReference() };
diff --git a/Engine/Plugins/Tests/RHITests/Source/RHITests/Public/RHITestsCommon.h b/Engine/Plugins/Tests/RHITests/Source/RHITests/Public/RHITestsCommon.h
--- a/Engine/Plugins/Tests/RHITests/Source/RHITests/Public/RHITestsCommon.h
+++ b/Engine/Plugins/Tests/RHITests/Source/RHITests/Public/RHITestsCommon.h
@@ -5,6 +5,7 @@
 #include "CoreMinimal.h"
 #include "RenderUtils.h"
 #include "RenderingThread.h"
+#include "RHIResourceUtils.h"
 
 DECLARE_LOG_CATEGORY_EXTERN(LogRHIUnitTestCommandlet, Log, All);
 
@@ -26,22 +27,3 @@ static inline FString ClearValueToString(const ValueType & ClearValue)
 		return FString::Printf(TEXT("0x%08x 0x%08x 0x%08x 0x%08x"), ClearValue.X, ClearValue.Y, ClearValue.Z, ClearValue.W);
 	}
 }
-
-template <typename T>
-static FBufferRHIRef CreateBufferWithData(EBufferUsageFlags UsageFlags, ERHIAccess ResourceState, const TCHAR* Name, TConstArrayView<T> Data)
-{
-	FRHICommandListBase& RHICmdList = FRHICommandListImmediate::Get();
-	uint32 BufferSize = sizeof(T) * Data.Num();
-	FRHIResourceCreateInfo CreateInfo(Name);
-	FBufferRHIRef Buffer = RHICmdList.CreateBuffer(BufferSize, UsageFlags, sizeof(T), ResourceState, CreateInfo);
-	void* MappedData = RHICmdList.LockBuffer(Buffer, 0, BufferSize, EResourceLockMode::RLM_WriteOnly);
-	FMemory::Memcpy(MappedData, Data.GetData(), BufferSize);
-	RHICmdList.UnlockBuffer(Buffer);
-	return Buffer;
-}
-
-template<typename T>
-static FBufferRHIRef CreateBufferWithData(EBufferUsageFlags UsageFlags, ERHIAccess ResourceState, const TCHAR* Name, TArrayView<T> Data)
-{
-	return CreateBufferWithData(UsageFlags, ResourceState, Name, TConstArrayView<T>(Data.GetData(), Data.Num()));
-}
diff --git a/Engine/Source/Runtime/Engine/Private/DynamicMeshBuilder.cpp b/Engine/Source/Runtime/Engine/Private/DynamicMeshBuilder.cpp
--- a/Engine/Source/Runtime/Engine/Private/DynamicMeshBuilder.cpp
+++ b/Engine/Source/Runtime/Engine/Private/DynamicMeshBuilder.cpp
@@ -11,6 +11,7 @@
 #include "Math/Vector2DHalf.h"
 #include "ResourcePool.h"
 #include "PrimitiveUniformShaderParametersBuilder.h"
+#include "RHIResourceUtils.h"
 
 class FGlobalDynamicMeshPoolPolicy
 {
@@ -198,24 +199,14 @@ TGlobalResource<FGlobalDynamicMeshVertexPool> GDynamicMeshVertexPool;
 
 void FDynamicMeshIndexBuffer32::InitRHI(FRHICommandListBase& RHICmdList)
 {
-	FRHIResourceCreateInfo CreateInfo(TEXT("FDynamicMeshIndexBuffer32"));
-	IndexBufferRHI = RHICmdList.CreateIndexBuffer(sizeof(uint32), Indices.Num() * sizeof(uint32), BUF_Static, CreateInfo);
-
-	// Copy the index data into the index buffer.
-	void* Buffer = RHICmdList.LockBuffer(IndexBufferRHI, 0, Indices.Num() * sizeof(uint32), RLM_WriteOnly);
-	FMemory::Memcpy(Buffer, Indices.GetData(), Indices.Num() * sizeof(uint32));
-	RHICmdList.UnlockBuffer(IndexBufferRHI);
+	FRHIResourceCreateInfoUploadArray CreateInfo(TEXT("FDynamicMeshIndexBuffer32"), MakeArrayView(Indices));
+	IndexBufferRHI = RHICmdList.CreateIndexBuffer(sizeof(uint32), CreateInfo.GetResourceDataSize(), BUF_Static, CreateInfo);
 }
 
 void FDynamicMeshIndexBuffer16::InitRHI(FRHICommandListBase& RHICmdList)
 {
-	FRHIResourceCreateInfo CreateInfo(TEXT("FDynamicMeshIndexBuffer16"));
-	IndexBufferRHI = RHICmdList.CreateIndexBuffer(sizeof(uint16), Indices.Num() * sizeof(uint16), BUF_Static, CreateInfo);
-
-	// Copy the index data into the index buffer.
-	void* Buffer = RHICmdList.LockBuffer(IndexBufferRHI, 0, Indices.Num() * sizeof(uint16), RLM_WriteOnly);
-	FMemory::Memcpy(Buffer, Indices.GetData(), Indices.Num() * sizeof(uint16));
-	RHICmdList.UnlockBuffer(IndexBufferRHI);
+	FRHIResourceCreateInfoUploadArray CreateInfo(TEXT("FDynamicMeshIndexBuffer16"), MakeArrayView(Indices));
+	IndexBufferRHI = RHICmdList.CreateIndexBuffer(sizeof(uint16), CreateInfo.GetResourceDataSize(), BUF_Static, CreateInfo);
 }
 
 /** FDynamicMeshBufferAllocator's base implementation. It always reallocates new buffers. */
diff --git a/Engine/Source/Runtime/Engine/Private/Particles/SubUVAnimation.cpp b/Engine/Source/Runtime/Engine/Private/Particles/SubUVAnimation.cpp
--- a/Engine/Source/Runtime/Engine/Private/Particles/SubUVAnimation.cpp
+++ b/Engine/Source/Runtime/Engine/Private/Particles/SubUVAnimation.cpp
@@ -12,6 +12,7 @@
 #include "DerivedDataCacheKey.h"
 #include "ComponentReregisterContext.h"
 #include "RHI.h"
+#include "RHIResourceUtils.h"
 
 #define LOCTEXT_NAMESPACE "FSubUVDerivedData"
 
@@ -59,13 +60,10 @@ FSubUVBoundingGeometryBuffer::~FSubUVBoundingGeometryBuffer() = default;
 
 void FSubUVBoundingGeometryBuffer::InitRHI(FRHICommandListBase& RHICmdList)
 {
-	const uint32 SizeInBytes = Vertices->Num() * Vertices->GetTypeSize();
-
-	if (SizeInBytes > 0)
+	if (Vertices->Num() > 0)
 	{
-		FResourceArrayUploadArrayView ResourceArray(Vertices->GetData(), SizeInBytes);
-		FRHIResourceCreateInfo CreateInfo(TEXT("FSubUVBoundingGeometryBuffer"), &ResourceArray);
-		VertexBufferRHI = RHICmdList.CreateVertexBuffer(SizeInBytes, BUF_ShaderResource | BUF_Static, CreateInfo);
+		FRHIResourceCreateInfoUploadArray CreateInfo(TEXT("FSubUVBoundingGeometryBuffer"), *Vertices);
+		VertexBufferRHI = RHICmdList.CreateVertexBuffer(CreateInfo.GetResourceDataSize(), BUF_ShaderResource | BUF_Static, CreateInfo);
 		ShaderResourceView = RHICmdList.CreateShaderResourceView(VertexBufferRHI, sizeof(FVector2f), PF_G32R32F);
 	}
 }
diff --git a/Engine/Source/Runtime/MRMesh/Private/MRMeshComponent.cpp b/Engine/Source/Runtime/MRMesh/Private/MRMeshComponent.cpp
--- a/Engine/Source/Runtime/MRMesh/Private/MRMeshComponent.cpp
+++ b/Engine/Source/Runtime/MRMesh/Private/MRMeshComponent.cpp
@@ -10,6 +10,7 @@
 #include "MaterialShared.h"
 #include "Materials/Material.h"
 #include "RenderingThread.h"
+#include "RHIResourceUtils.h"
 #include "DataDrivenShaderPlatformInfo.h"
 #include "BaseMeshReconstructorModule.h"
 #include "MeshReconstructorBase.h"
@@ -60,11 +61,8 @@ public:
 	{
 		NumVerts = PerVertexData.Num();
 
-		const uint32 SizeInBytes = PerVertexData.Num() * sizeof(DataType);
-
-		FResourceArrayUploadArrayView ResourceArray(PerVertexData.GetData(), SizeInBytes);
-		FRHIResourceCreateInfo CreateInfo(TEXT("FMRMeshVertexBuffer"), &ResourceArray);
-		VertexBufferRHI = RHICmdList.CreateVertexBuffer(SizeInBytes, BUF_Static | BUF_ShaderResource, CreateInfo);
+		FRHIResourceCreateInfoUploadArray CreateInfo(TEXT("FMRMeshVertexBuffer"), PerVertexData);
+		VertexBufferRHI = RHICmdList.CreateVertexBuffer(CreateInfo.GetResourceDataSize(), BUF_Static | BUF_ShaderResource, CreateInfo);
 	}
 
 };
@@ -77,30 +75,16 @@ public:
 	{
 		NumIndices = Indices.Num();
 
-		const uint32 Size = Indices.Num() * sizeof(uint32);
-
-		FRHIResourceCreateInfo CreateInfo(TEXT("FMRMeshIndexBuffer"));
-		IndexBufferRHI = RHICmdList.CreateBuffer(Size, BUF_Static | BUF_IndexBuffer, sizeof(uint32), ERHIAccess::VertexOrIndexBuffer, CreateInfo);
-
-		// Write the indices to the index buffer.
-		void* Buffer = RHICmdList.LockBuffer(IndexBufferRHI, 0, Size, RLM_WriteOnly);
-		FMemory::Memcpy(Buffer, Indices.GetData(), Size);
-		RHICmdList.UnlockBuffer(IndexBufferRHI);
+		FRHIResourceCreateInfoUploadArray CreateInfo(TEXT("FMRMeshIndexBuffer"), Indices);
+		IndexBufferRHI = RHICmdList.CreateBuffer(CreateInfo.GetResourceDataSize(), BUF_Static | BUF_IndexBuffer, sizeof(uint32), ERHIAccess::VertexOrIndexBuffer, CreateInfo);
 	}
 
 	void InitRHIWith(FRHICommandListBase& RHICmdList, const TArray<uint16>& Indices)
 	{
 		NumIndices = Indices.Num();
 
-		const uint32 Size = Indices.Num() * sizeof(uint16);
-
-		FRHIResourceCreateInfo CreateInfo(TEXT("FMRMeshIndexBuffer"));
-		IndexBufferRHI = RHICmdList.CreateBuffer(Size, BUF_Static | BUF_IndexBuffer, sizeof(uint16), ERHIAccess::VertexOrIndexBuffer, CreateInfo);
-
-		// Write the indices to the index buffer.
-		void* Buffer = RHICmdList.LockBuffer(IndexBufferRHI, 0, Size, RLM_WriteOnly);
-		FMemory::Memcpy(Buffer, Indices.GetData(), Size);
-		RHICmdList.UnlockBuffer(IndexBufferRHI);
+		FRHIResourceCreateInfoUploadArray CreateInfo(TEXT("FMRMeshIndexBuffer"), Indices);
+		IndexBufferRHI = RHICmdList.CreateBuffer(CreateInfo.GetResourceDataSize(), BUF_Static | BUF_IndexBuffer, sizeof(uint16), ERHIAccess::VertexOrIndexBuffer, CreateInfo);
 	}
 };
 
diff --git a/Engine/Source/Runtime/RHI/Public/RHIResourceUtils.h b/Engine/Source/Runtime/RHI/Public/RHIResourceUtils.h
new file mode 100644
--- /dev/null
+++ b/Engine/Source/Runtime/RHI/Public/RHIResourceUtils.h
@@ -0,0 +1,59 @@
+// Copyright Epic Games, Inc. All Rights Reserved.
+
+#pragma once
+
+#include "RHICommandList.h"
+#include "RHIResources.h"
+#include "Containers/ResourceArray.h"
+
+struct FRHIResourceCreateInfoUploadArray : public FRHIResourceCreateInfo
+{
+	FResourceArrayUploadArrayView UploadView;
+
+	FRHIResourceCreateInfoUploadArray(const TCHAR* InDebugName, const void* InData, uint32 InSizeInBytes)
+		: FRHIResourceCreateInfo(InDebugName, &UploadView)
+		, UploadView(InData, InSizeInBytes)
+	{
+	}
+
+	template<typename ElementType>
+	FRHIResourceCreateInfoUploadArray(const TCHAR* InDebugName, TConstArrayView<ElementType> InView)
+		: FRHIResourceCreateInfo(InDebugName, &UploadView)
+		, UploadView(InView)
+	{
+	}
+
+	template<typename ElementType>
+	FRHIResourceCreateInfoUploadArray(const TCHAR* InDebugName, TArrayView<ElementType> InView)
+	: FRHIResourceCreateInfoUploadArray(InDebugName, TConstArrayView<ElementType>(InView))
+	{
+	}
+
+	template<typename ElementType, typename AllocatorType>
+	FRHIResourceCreateInfoUploadArray(const TCHAR* InDebugName, const TArray<ElementType, AllocatorType>& InArray)
+	: FRHIResourceCreateInfoUploadArray(InDebugName, TConstArrayView<ElementType>(InArray))
+	{
+	}
+
+	uint32 GetResourceDataSize() const
+	{
+		return UploadView.GetResourceDataSize();
+	}
+};
+
+namespace UE::RHIResourceUtils
+{
+	template <typename T>
+	static FBufferRHIRef CreateBufferWithData(FRHICommandListBase& RHICmdList, EBufferUsageFlags UsageFlags, ERHIAccess ResourceState, const TCHAR* Name, TConstArrayView<T> Data)
+	{
+		FRHIResourceCreateInfoUploadArray CreateInfo(Name, Data);
+		return RHICmdList.CreateBuffer(CreateInfo.GetResourceDataSize(), UsageFlags, Data.GetTypeSize(), ResourceState, CreateInfo);
+	}
+
+	template<typename T>
+	static FBufferRHIRef CreateBufferWithData(FRHICommandListBase& RHICmdList, EBufferUsageFlags UsageFlags, ERHIAccess ResourceState, const TCHAR* Name, TArrayView<T> Data)
+	{
+		return CreateBufferWithData(RHICmdList, UsageFlags, ResourceState, Name, TConstArrayView<T>(Data));
+	}
+}
+
diff --git a/Engine/Source/Runtime/RenderCore/Private/CommonRenderResources.cpp b/Engine/Source/Runtime/RenderCore/Private/CommonRenderResources.cpp
--- a/Engine/Source/Runtime/RenderCore/Private/CommonRenderResources.cpp
+++ b/Engine/Source/Runtime/RenderCore/Private/CommonRenderResources.cpp
@@ -6,6 +6,7 @@
 
 #include "CommonRenderResources.h"
 #include "Containers/DynamicRHIResourceArray.h"
+#include "RHIResourceUtils.h"
 #include "StereoRenderUtils.h"
 
 
@@ -34,50 +35,35 @@ bool FMobileMultiViewVertexShaderVS::ShouldCompilePermutation(const FGlobalShade
 	return Aspects.IsMobileMultiViewEnabled();
 }
 
-void FScreenRectangleVertexBuffer::InitRHI(FRHICommandListBase& RHICmdList)
+static const FFilterVertex GScreenRectangleVertexBufferData[] =
 {
-	TResourceArray<FFilterVertex, VERTEXBUFFER_ALIGNMENT> Vertices;
-	Vertices.SetNumUninitialized(6);
-
-	Vertices[0].Position = FVector4f(1, 1, 0, 1);
-	Vertices[0].UV = FVector2f(1, 1);
-
-	Vertices[1].Position = FVector4f(0, 1, 0, 1);
-	Vertices[1].UV = FVector2f(0, 1);
-
-	Vertices[2].Position = FVector4f(1, 0, 0, 1);
-	Vertices[2].UV = FVector2f(1, 0);
-
-	Vertices[3].Position = FVector4f(0, 0, 0, 1);
-	Vertices[3].UV = FVector2f(0, 0);
+	{ FVector4f(1, 1, 0, 1), FVector2f(1, 1) },
+	{ FVector4f(0, 1, 0, 1), FVector2f(0, 1) },
+	{ FVector4f(1, 0, 0, 1), FVector2f(1, 0) },
+	{ FVector4f(0, 0, 0, 1), FVector2f(0, 0) },
 
 	//The final two vertices are used for the triangle optimization (a single triangle spans the entire viewport )
-	Vertices[4].Position = FVector4f(-1, 1, 0, 1);
-	Vertices[4].UV = FVector2f(-1, 1);
-
-	Vertices[5].Position = FVector4f(1, -1, 0, 1);
-	Vertices[5].UV = FVector2f(1, -1);
+	{ FVector4f(-1, 1, 0, 1), FVector2f(-1, 1) },
+	{ FVector4f(1, -1, 0, 1), FVector2f(1, -1) },
+};
 
+void FScreenRectangleVertexBuffer::InitRHI(FRHICommandListBase& RHICmdList)
+{
 	// Create vertex buffer. Fill buffer with initial data upon creation
-	FRHIResourceCreateInfo CreateInfo(TEXT("FScreenRectangleVertexBuffer"), &Vertices);
-	VertexBufferRHI = RHICmdList.CreateVertexBuffer(Vertices.GetResourceDataSize(), BUF_Static, CreateInfo);
+	FRHIResourceCreateInfoUploadArray CreateInfo(TEXT("FScreenRectangleVertexBuffer"), MakeArrayView(GScreenRectangleVertexBufferData));
+	VertexBufferRHI = RHICmdList.CreateVertexBuffer(CreateInfo.GetResourceDataSize(), BUF_Static, CreateInfo);
 }
 
-void FScreenRectangleIndexBuffer::InitRHI(FRHICommandListBase& RHICmdList)
+static const uint16 GScreenRectangleIndexBufferData[] =
 {
-	const uint16 Indices[] = 
-	{
-		0, 1, 2, 2, 1, 3,	// [0 .. 5]  Full screen quad with 2 triangles
-		0, 4, 5,			// [6 .. 8]  Full screen triangle
-		3, 2, 1				// [9 .. 11] Full screen rect defined with TL, TR, BL corners
-	};
-
-	TResourceArray<uint16, INDEXBUFFER_ALIGNMENT> IndexBuffer;
-	uint32 NumIndices = UE_ARRAY_COUNT(Indices);
-	IndexBuffer.AddUninitialized(NumIndices);
-	FMemory::Memcpy(IndexBuffer.GetData(), Indices, NumIndices * sizeof(uint16));
+	0, 1, 2, 2, 1, 3,	// [0 .. 5]  Full screen quad with 2 triangles
+	0, 4, 5,			// [6 .. 8]  Full screen triangle
+	3, 2, 1				// [9 .. 11] Full screen rect defined with TL, TR, BL corners
+};
 
+void FScreenRectangleIndexBuffer::InitRHI(FRHICommandListBase& RHICmdList)
+{
 	// Create index buffer. Fill buffer with initial data upon creation
-	FRHIResourceCreateInfo CreateInfo(TEXT("FScreenRectangleIndexBuffer"), &IndexBuffer);
-	IndexBufferRHI = RHICmdList.CreateIndexBuffer(sizeof(uint16), IndexBuffer.GetResourceDataSize(), BUF_Static, CreateInfo);
+	FRHIResourceCreateInfoUploadArray CreateInfo(TEXT("FScreenRectangleIndexBuffer"), MakeArrayView(GScreenRectangleIndexBufferData));
+	IndexBufferRHI = RHICmdList.CreateIndexBuffer(sizeof(uint16), CreateInfo.GetResourceDataSize(), BUF_Static, CreateInfo);
 }
diff --git a/Engine/Source/Runtime/RenderCore/Private/UnifiedBuffer.cpp b/Engine/Source/Runtime/RenderCore/Private/UnifiedBuffer.cpp
--- a/Engine/Source/Runtime/RenderCore/Private/UnifiedBuffer.cpp
+++ b/Engine/Source/Runtime/RenderCore/Private/UnifiedBuffer.cpp
@@ -8,6 +8,7 @@
 #include "Shader.h"
 #include "GlobalShader.h"
 #include "RenderGraphUtils.h"
+#include "RHIResourceUtils.h"
 #include "DataDrivenShaderPlatformInfo.h"
 
 // Uploads use a storage buffers which are at least 128m elements
@@ -1704,14 +1705,12 @@ void FScatterUploadBuffer::ResourceUploadTo(FRHICommandList& RHICmdList, const R
 		const EBufferUsageFlags Usage = bFloat4Buffer ? BUF_None : BUF_ByteAddressBuffer;
 
 		{
-			FResourceArrayUploadArrayView ScatterResourceArray(ScatterData, ScatterDataSize);
-			FRHIResourceCreateInfo CreateInfo(TEXT("ScatterResourceArray"), &ScatterResourceArray);
+			FRHIResourceCreateInfoUploadArray CreateInfo(TEXT("ScatterResourceArray"), ScatterData, ScatterDataSize);
 			ScatterBuffer.Buffer = RHICmdList.CreateStructuredBuffer(sizeof(uint32), ScatterDataSize, BUF_ShaderResource | BUF_Volatile | Usage, CreateInfo);
 			ScatterBuffer.SRV = RHICmdList.CreateShaderResourceView(ScatterBuffer.Buffer);
 		}
 		{
-			FResourceArrayUploadArrayView UploadResourceArray(UploadData, UploadDataSize);
-			FRHIResourceCreateInfo CreateInfo(TEXT("ScatterUploadBuffer"), &UploadResourceArray);
+			FRHIResourceCreateInfoUploadArray CreateInfo(TEXT("ScatterUploadBuffer"), UploadData, UploadDataSize);
 			UploadBuffer.Buffer = RHICmdList.CreateStructuredBuffer(TypeSize, UploadDataSize, BUF_ShaderResource | BUF_Volatile | Usage, CreateInfo);
 			UploadBuffer.SRV = RHICmdList.CreateShaderResourceView(UploadBuffer.Buffer);
 		}
diff --git a/Engine/Source/Runtime/Renderer/Private/InstanceCulling/InstanceCullingOcclusionQuery.cpp b/Engine/Source/Runtime/Renderer/Private/InstanceCulling/InstanceCullingOcclusionQuery.cpp
--- a/Engine/Source/Runtime/Renderer/Private/InstanceCulling/InstanceCullingOcclusionQuery.cpp
+++ b/Engine/Source/Runtime/Renderer/Private/InstanceCulling/InstanceCullingOcclusionQuery.cpp
@@ -9,6 +9,7 @@
 #include "RHIAccess.h"
 #include "RHIFeatureLevel.h"
 #include "RHIGlobals.h"
+#include "RHIResourceUtils.h"
 #include "RHIShaderPlatform.h"
 #include "RHIStaticStates.h"
 #include "RenderGraphBuilder.h"
@@ -33,14 +34,7 @@ struct FInstanceCullingOcclusionQueryDeferredContext;
 namespace
 {
 
-template <typename T>
-static FBufferRHIRef CreateBufferWithData(FRHICommandListBase& RHICmdList, EBufferUsageFlags UsageFlags, ERHIAccess ResourceState, const TCHAR* Name, TConstArrayView<T> Data)
-{
-	FResourceArrayUploadArrayView DataView(Data);
-	FRHIResourceCreateInfo CreateInfo(Name);
-	CreateInfo.ResourceArray = &DataView;
-	return RHICmdList.CreateBuffer(DataView.SizeInBytes, UsageFlags, Data.GetTypeSize(), ResourceState, CreateInfo);
-}
+
 
 
 static EPixelFormat GetPreferredVisibilityMaskFormat()
@@ -224,11 +218,11 @@ public:
 			FVector3f(+1.0f, -1.0f, -1.0f),
 		};
 
-		IndexBuffer = CreateBufferWithData(RHICmdList, EBufferUsageFlags::IndexBuffer, ERHIAccess::VertexOrIndexBuffer,
+		IndexBuffer = UE::RHIResourceUtils::CreateBufferWithData(RHICmdList, EBufferUsageFlags::IndexBuffer, ERHIAccess::VertexOrIndexBuffer,
 			TEXT("FInstanceCullingOcclusionQueryBox_IndexBuffer"),
 			MakeArrayView(BoxIndexBufferData));
 
-		VertexBuffer = CreateBufferWithData(RHICmdList, EBufferUsageFlags::VertexBuffer, ERHIAccess::VertexOrIndexBuffer,
+		VertexBuffer = UE::RHIResourceUtils::CreateBufferWithData(RHICmdList, EBufferUsageFlags::VertexBuffer, ERHIAccess::VertexOrIndexBuffer,
 			TEXT("FInstanceCullingOcclusionQueryBox_VertexBuffer"),
 			MakeArrayView(BoxVertexBufferData));
 
