Fixed missing test case for newly added PF_R16G16_SINT pixel format.

diff --git a/Engine/Plugins/Tests/RHITests/Source/RHITests/Public/RHITextureTests.h b/Engine/Plugins/Tests/RHITests/Source/RHITests/Public/RHITextureTests.h
--- a/Engine/Plugins/Tests/RHITests/Source/RHITests/Public/RHITextureTests.h
+++ b/Engine/Plugins/Tests/RHITests/Source/RHITests/Public/RHITextureTests.h
@@ -772,6 +772,7 @@ public:
 				RUN_TEST(Test_UpdateTexture_Impl<FDataSource16x2>(RHICmdList, Format, UpdateType));
 				break;
 			case PF_R16G16_UINT:
+			case PF_R16G16_SINT:
 				RUN_TEST(Test_UpdateTexture_Impl<FDataSource16x2UInt>(RHICmdList, Format, UpdateType));
 				break;
 			case PF_G16R16F:
