RHI Buffer upload helper changes
- CreateBufferWithData renamed to CreateBufferFromArray since it's using the array to set the buffer dimensions.
- Adding shortcuts to CreateBufferFromArray: CreateIndexBufferFromArray and CreateVertexBufferFromArray

diff --git a/Engine/Plugins/Tests/RHITests/Source/RHITests/Private/RHIDrawTests.cpp b/Engine/Plugins/Tests/RHITests/Source/RHITests/Private/RHIDrawTests.cpp
--- a/Engine/Plugins/Tests/RHITests/Source/RHITests/Private/RHIDrawTests.cpp
+++ b/Engine/Plugins/Tests/RHITests/Source/RHITests/Private/RHIDrawTests.cpp
@@ -74,7 +74,7 @@ struct FDrawTestResources
 		{
 			InstanceIDs[i] = i;
 		}
-		InstanceIDBuffer = UE::RHIResourceUtils::CreateBufferWithData(RHICmdList, EBufferUsageFlags::VertexBuffer, ERHIAccess::VertexOrIndexBuffer, TEXT("DrawTest_InstanceID"), MakeArrayView(InstanceIDs));
+		InstanceIDBuffer = UE::RHIResourceUtils::CreateVertexBufferFromArray(RHICmdList, TEXT("DrawTest_InstanceID"), MakeConstArrayView(InstanceIDs));
 
 		// Indices for 4 triangles
 		const uint16 Indices[NumTotalVertices] =
@@ -85,7 +85,7 @@ struct FDrawTestResources
 			9, 10, 11 // degenerate
 		};
 
-		IndexBuffer = UE::RHIResourceUtils::CreateBufferWithData(RHICmdList, EBufferUsageFlags::IndexBuffer, ERHIAccess::VertexOrIndexBuffer, TEXT("DrawTest_IndexBuffer"), MakeArrayView(Indices));
+		IndexBuffer = UE::RHIResourceUtils::CreateIndexBufferFromArray(RHICmdList, TEXT("DrawTest_IndexBuffer"), MakeConstArrayView(Indices));
 
 		TArray<FVector4f> Vertices;
 		Vertices.Reserve(NumTotalVertices);
@@ -102,7 +102,7 @@ struct FDrawTestResources
 		{
 			Vertices.Add(FVector4f(0.0f, 0.0f, 0.0f, 1.0f));
 		}
-		VertexBuffer = UE::RHIResourceUtils::CreateBufferWithData(RHICmdList, EBufferUsageFlags::VertexBuffer, ERHIAccess::VertexOrIndexBuffer, TEXT("DrawTest_VertexBuffer"), MakeArrayView(Vertices));
+		VertexBuffer = UE::RHIResourceUtils::CreateVertexBufferFromArray(RHICmdList, TEXT("DrawTest_VertexBuffer"), MakeConstArrayView(Vertices));
 
 		static constexpr uint32 OutputBufferStride = sizeof(uint32);
 		static constexpr uint32 OutputBufferSize = OutputBufferStride * MaxInstances;
@@ -167,8 +167,13 @@ bool FRHIDrawTests::InternalDrawBaseVertexAndInstance(FRHICommandListImmediate&
 	FBufferRHIRef DrawArgBuffer;
 	if (DrawKind == EDrawKind::Indirect)
 	{
-		DrawArgBuffer = UE::RHIResourceUtils::CreateBufferWithData(RHICmdList, EBufferUsageFlags::DrawIndirect | EBufferUsageFlags::UnorderedAccess | EBufferUsageFlags::VertexBuffer, ERHIAccess::IndirectArgs,
-			TEXT("InternalDrawBaseVertexAndInstance_DrawArgs"), MakeArrayView(DrawArgs));
+		DrawArgBuffer = UE::RHIResourceUtils::CreateBufferFromArray(
+			RHICmdList,
+			TEXT("InternalDrawBaseVertexAndInstance_DrawArgs"), 
+			EBufferUsageFlags::DrawIndirect | EBufferUsageFlags::UnorderedAccess | EBufferUsageFlags::VertexBuffer,
+			ERHIAccess::IndirectArgs,
+			MakeConstArrayView(DrawArgs)
+		);
 	}
 
 	RHICmdList.Transition(FRHITransitionInfo(Resources.OutputBufferUAV, ERHIAccess::UAVCompute, ERHIAccess::UAVGraphics, EResourceTransitionFlags::None));
@@ -256,7 +261,13 @@ bool FRHIDrawTests::Test_MultiDrawIndirect(FRHICommandListImmediate& RHICmdList)
 	FDrawTestResources Resources(RHICmdList);
 
 	const uint32 CountValues[4] = { 1, 1, 16, 0 };
-	FBufferRHIRef CountBuffer = UE::RHIResourceUtils::CreateBufferWithData(RHICmdList, EBufferUsageFlags::DrawIndirect | EBufferUsageFlags::UnorderedAccess, ERHIAccess::IndirectArgs, TEXT("Test_MultiDrawIndirect_Count"), MakeArrayView(CountValues));
+	FBufferRHIRef CountBuffer = UE::RHIResourceUtils::CreateBufferFromArray(
+		RHICmdList,
+		TEXT("Test_MultiDrawIndirect_Count"),
+		EBufferUsageFlags::DrawIndirect | EBufferUsageFlags::UnorderedAccess,
+		ERHIAccess::IndirectArgs,
+		MakeConstArrayView(CountValues)
+	);
 
 	const FRHIDrawIndexedIndirectParameters DrawArgs[] =
 	{
@@ -272,8 +283,13 @@ bool FRHIDrawTests::Test_MultiDrawIndirect(FRHICommandListImmediate& RHICmdList)
 
 	const uint32 ExpectedDrawnInstances[Resources.MaxInstances] = { 1, 0, 1, 1, 0, 1, 1, 0 };
 
-	FBufferRHIRef DrawArgBuffer = UE::RHIResourceUtils::CreateBufferWithData(RHICmdList, EBufferUsageFlags::DrawIndirect | EBufferUsageFlags::UnorderedAccess | EBufferUsageFlags::VertexBuffer, ERHIAccess::IndirectArgs,
-		TEXT("Test_MultiDrawIndirect_DrawArgs"), MakeArrayView(DrawArgs));
+	FBufferRHIRef DrawArgBuffer = UE::RHIResourceUtils::CreateBufferFromArray(
+		RHICmdList,
+		TEXT("Test_MultiDrawIndirect_DrawArgs"),
+		EBufferUsageFlags::DrawIndirect | EBufferUsageFlags::UnorderedAccess | EBufferUsageFlags::VertexBuffer,
+		ERHIAccess::IndirectArgs,
+		MakeConstArrayView(DrawArgs)
+	);
 
 	RHICmdList.ClearUAVUint(Resources.OutputBufferUAV, FUintVector4(0));
 
diff --git a/Engine/Plugins/Tests/RHITests/Source/RHITests/Private/RHIGraphicsUAVTests.cpp b/Engine/Plugins/Tests/RHITests/Source/RHITests/Private/RHIGraphicsUAVTests.cpp
--- a/Engine/Plugins/Tests/RHITests/Source/RHITests/Private/RHIGraphicsUAVTests.cpp
+++ b/Engine/Plugins/Tests/RHITests/Source/RHITests/Private/RHIGraphicsUAVTests.cpp
@@ -89,8 +89,7 @@ static void Test_GraphicsUAV_Common(FRHICommandListImmediate& RHICmdList,
 	Vertices.Add(FVector4f(-1.0f, +3.0f, 0.0f, 1.0f));
 	Vertices.Add(FVector4f(+3.0f, -1.0f, 0.0f, 1.0f));
 
-	FBufferRHIRef VertexBuffer = UE::RHIResourceUtils::CreateBufferWithData(RHICmdList, EBufferUsageFlags::VertexBuffer,
-		ERHIAccess::VertexOrIndexBuffer, TEXT("GraphicsUAVTests_VertexBuffer"), MakeArrayView(Vertices));
+	FBufferRHIRef VertexBuffer = UE::RHIResourceUtils::CreateVertexBufferFromArray(RHICmdList, TEXT("GraphicsUAVTests_VertexBuffer"), MakeConstArrayView(Vertices));
 
 	FRHITexture* ColorRTs[1] = { RenderTarget.GetReference() };
 	FRHIRenderPassInfo RenderPassInfo(1, ColorRTs, ERenderTargetActions::DontLoad_DontStore);
diff --git a/Engine/Source/Runtime/MRMesh/Private/MRMeshComponent.cpp b/Engine/Source/Runtime/MRMesh/Private/MRMeshComponent.cpp
--- a/Engine/Source/Runtime/MRMesh/Private/MRMeshComponent.cpp
+++ b/Engine/Source/Runtime/MRMesh/Private/MRMeshComponent.cpp
@@ -60,11 +60,8 @@ public:
 	void InitRHIWith(FRHICommandListBase& RHICmdList, const TArray<DataType>& PerVertexData )
 	{
 		NumVerts = PerVertexData.Num();
-
-		FRHIResourceCreateInfoUploadArray CreateInfo(TEXT("FMRMeshVertexBuffer"), PerVertexData);
-		VertexBufferRHI = RHICmdList.CreateVertexBuffer(CreateInfo.GetResourceDataSize(), BUF_Static | BUF_ShaderResource, CreateInfo);
+		VertexBufferRHI = UE::RHIResourceUtils::CreateVertexBufferFromArray(RHICmdList, TEXT("FMRMeshVertexBuffer"), EBufferUsageFlags::Static, MakeConstArrayView(PerVertexData));
 	}
-
 };
 
 class FMRMeshIndexBuffer : public FIndexBuffer
@@ -74,17 +71,13 @@ public:
 	void InitRHIWith(FRHICommandListBase& RHICmdList, const TArray<uint32>& Indices )
 	{
 		NumIndices = Indices.Num();
-
-		FRHIResourceCreateInfoUploadArray CreateInfo(TEXT("FMRMeshIndexBuffer"), Indices);
-		IndexBufferRHI = RHICmdList.CreateBuffer(CreateInfo.GetResourceDataSize(), BUF_Static | BUF_IndexBuffer, sizeof(uint32), ERHIAccess::VertexOrIndexBuffer, CreateInfo);
+		IndexBufferRHI = UE::RHIResourceUtils::CreateIndexBufferFromArray(RHICmdList, TEXT("FMRMeshIndexBuffer"), EBufferUsageFlags::Static, MakeConstArrayView(Indices));
 	}
 
 	void InitRHIWith(FRHICommandListBase& RHICmdList, const TArray<uint16>& Indices)
 	{
 		NumIndices = Indices.Num();
-
-		FRHIResourceCreateInfoUploadArray CreateInfo(TEXT("FMRMeshIndexBuffer"), Indices);
-		IndexBufferRHI = RHICmdList.CreateBuffer(CreateInfo.GetResourceDataSize(), BUF_Static | BUF_IndexBuffer, sizeof(uint16), ERHIAccess::VertexOrIndexBuffer, CreateInfo);
+		IndexBufferRHI = UE::RHIResourceUtils::CreateIndexBufferFromArray(RHICmdList, TEXT("FMRMeshIndexBuffer"), EBufferUsageFlags::Static, MakeConstArrayView(Indices));
 	}
 };
 
diff --git a/Engine/Source/Runtime/RHI/Public/RHIResourceUtils.h b/Engine/Source/Runtime/RHI/Public/RHIResourceUtils.h
--- a/Engine/Source/Runtime/RHI/Public/RHIResourceUtils.h
+++ b/Engine/Source/Runtime/RHI/Public/RHIResourceUtils.h
@@ -43,17 +43,34 @@ struct FRHIResourceCreateInfoUploadArray : public FRHIResourceCreateInfo
 
 namespace UE::RHIResourceUtils
 {
-	template <typename T>
-	static FBufferRHIRef CreateBufferWithData(FRHICommandListBase& RHICmdList, EBufferUsageFlags UsageFlags, ERHIAccess ResourceState, const TCHAR* Name, TConstArrayView<T> Data)
+	template <typename TElementType>
+	static FBufferRHIRef CreateBufferFromArray(FRHICommandListBase& RHICmdList, const TCHAR* Name, EBufferUsageFlags UsageFlags, ERHIAccess ResourceState, TConstArrayView<TElementType> Array)
 	{
-		FRHIResourceCreateInfoUploadArray CreateInfo(Name, Data);
-		return RHICmdList.CreateBuffer(CreateInfo.GetResourceDataSize(), UsageFlags, Data.GetTypeSize(), ResourceState, CreateInfo);
+		FRHIResourceCreateInfoUploadArray CreateInfo(Name, Array);
+		return RHICmdList.CreateBuffer(CreateInfo.GetResourceDataSize(), UsageFlags, Array.GetTypeSize(), ResourceState, CreateInfo);
 	}
 
-	template<typename T>
-	static FBufferRHIRef CreateBufferWithData(FRHICommandListBase& RHICmdList, EBufferUsageFlags UsageFlags, ERHIAccess ResourceState, const TCHAR* Name, TArrayView<T> Data)
+	template<typename TElementType>
+	static FBufferRHIRef CreateVertexBufferFromArray(FRHICommandListBase& RHICmdList, const TCHAR* Name, EBufferUsageFlags ExtraFlags, TConstArrayView<TElementType> Array)
 	{
-		return CreateBufferWithData(RHICmdList, UsageFlags, ResourceState, Name, TConstArrayView<T>(Data));
+		return CreateBufferFromArray<TElementType>(RHICmdList, Name, EBufferUsageFlags::VertexBuffer | ExtraFlags, ERHIAccess::VertexOrIndexBuffer, Array);
+	}
+
+	template<typename TElementType>
+	static FBufferRHIRef CreateVertexBufferFromArray(FRHICommandListBase& RHICmdList, const TCHAR* Name, TConstArrayView<TElementType> Array)
+	{
+		return CreateVertexBufferFromArray<TElementType>(RHICmdList, Name, EBufferUsageFlags::None, Array);
+	}
+
+	template<typename TElementType>
+	static FBufferRHIRef CreateIndexBufferFromArray(FRHICommandListBase& RHICmdList, const TCHAR* Name, EBufferUsageFlags ExtraFlags, TConstArrayView<TElementType> Array)
+	{
+		return CreateBufferFromArray<TElementType>(RHICmdList, Name, EBufferUsageFlags::IndexBuffer | ExtraFlags, ERHIAccess::VertexOrIndexBuffer, Array);
 	}
-}
 
+	template<typename TElementType>
+	static FBufferRHIRef CreateIndexBufferFromArray(FRHICommandListBase& RHICmdList, const TCHAR* Name, TConstArrayView<TElementType> Array)
+	{
+		return CreateIndexBufferFromArray<TElementType>(RHICmdList, Name, EBufferUsageFlags::None, Array);
+	}
+}
diff --git a/Engine/Source/Runtime/Renderer/Private/InstanceCulling/InstanceCullingOcclusionQuery.cpp b/Engine/Source/Runtime/Renderer/Private/InstanceCulling/InstanceCullingOcclusionQuery.cpp
--- a/Engine/Source/Runtime/Renderer/Private/InstanceCulling/InstanceCullingOcclusionQuery.cpp
+++ b/Engine/Source/Runtime/Renderer/Private/InstanceCulling/InstanceCullingOcclusionQuery.cpp
@@ -218,13 +218,8 @@ public:
 			FVector3f(+1.0f, -1.0f, -1.0f),
 		};
 
-		IndexBuffer = UE::RHIResourceUtils::CreateBufferWithData(RHICmdList, EBufferUsageFlags::IndexBuffer, ERHIAccess::VertexOrIndexBuffer,
-			TEXT("FInstanceCullingOcclusionQueryBox_IndexBuffer"),
-			MakeArrayView(BoxIndexBufferData));
-
-		VertexBuffer = UE::RHIResourceUtils::CreateBufferWithData(RHICmdList, EBufferUsageFlags::VertexBuffer, ERHIAccess::VertexOrIndexBuffer,
-			TEXT("FInstanceCullingOcclusionQueryBox_VertexBuffer"),
-			MakeArrayView(BoxVertexBufferData));
+		IndexBuffer = UE::RHIResourceUtils::CreateIndexBufferFromArray(RHICmdList, TEXT("FInstanceCullingOcclusionQueryBox_IndexBuffer"), MakeConstArrayView(BoxIndexBufferData));
+		VertexBuffer = UE::RHIResourceUtils::CreateVertexBufferFromArray(RHICmdList, TEXT("FInstanceCullingOcclusionQueryBox_VertexBuffer"), MakeConstArrayView(BoxVertexBufferData));
 
 		FVertexDeclarationElementList VertexDeclarationElements;
 		VertexDeclarationElements.Add(FVertexElement(0, 0, VET_Float3, 0, 12));
