Remove uncorrect check on the stride size from the unit test, since the stride information is in vertex layout, only structured buffer should have the stride size in the buffer.

diff --git a/Engine/Plugins/Tests/RHITests/Source/RHITests/Private/RHIDrawTests.cpp b/Engine/Plugins/Tests/RHITests/Source/RHITests/Private/RHIDrawTests.cpp
--- a/Engine/Plugins/Tests/RHITests/Source/RHITests/Private/RHIDrawTests.cpp
+++ b/Engine/Plugins/Tests/RHITests/Source/RHITests/Private/RHIDrawTests.cpp
@@ -185,7 +185,6 @@ bool FRHIDrawTests::InternalDrawBaseVertexAndInstance(FRHICommandListImmediate&
 	RHICmdList.ApplyCachedRenderTargets(Resources.GraphicsPSOInit);
 	SetGraphicsPipelineState(RHICmdList, Resources.GraphicsPSOInit, 0);
 
-	check(Resources.InstanceIDBuffer->GetStride() == 4);
 	RHICmdList.SetStreamSource(0, Resources.VertexBuffer, 0);
 	RHICmdList.SetStreamSource(1, Resources.InstanceIDBuffer, 0);
 
@@ -305,7 +304,6 @@ bool FRHIDrawTests::Test_MultiDrawIndirect(FRHICommandListImmediate& RHICmdList)
 	RHICmdList.ApplyCachedRenderTargets(Resources.GraphicsPSOInit);
 	SetGraphicsPipelineState(RHICmdList, Resources.GraphicsPSOInit, 0);
 
-	check(Resources.InstanceIDBuffer->GetStride() == 4);
 	RHICmdList.SetStreamSource(0, Resources.VertexBuffer, 0);
 	RHICmdList.SetStreamSource(1, Resources.InstanceIDBuffer, 0);
 
