Increased the size of buffer to account for minimum alignment on Polaris Mac

https://horde.devtools.epicgames.com/testreport/66ff7f697e36cf372d4d2107

diff --git a/Engine/Plugins/Tests/RHITests/Source/RHITests/Public/RHIBufferTests.h b/Engine/Plugins/Tests/RHITests/Source/RHITests/Public/RHIBufferTests.h
--- a/Engine/Plugins/Tests/RHITests/Source/RHITests/Public/RHIBufferTests.h
+++ b/Engine/Plugins/Tests/RHITests/Source/RHITests/Public/RHIBufferTests.h
@@ -238,65 +238,65 @@ public:
 			const FUintVector4 ClearValueUint16(0x0123, 0x4567, 0x89ab, 0xcdef);
 			const FUintVector4 ClearValueUint32(0x01234567, 0x89abcdef, 0x8899aabb, 0xccddeeff);
 
-			RUN_TEST(RunTest_UAVClear_Buffer(RHICmdList, BUF_VertexBuffer, 512, PF_R8_UINT          , 0, false, ClearValueUint8 , &FRHICommandListImmediate::ClearUAVUint, { 0x01 }));
-			RUN_TEST(RunTest_UAVClear_Buffer(RHICmdList, BUF_VertexBuffer, 512, PF_R8G8B8A8_UINT    , 0, false, ClearValueUint8 , &FRHICommandListImmediate::ClearUAVUint, { 0x01, 0x23, 0x45, 0x67 }));
-			RUN_TEST(RunTest_UAVClear_Buffer(RHICmdList, BUF_VertexBuffer, 512, PF_R16_UINT         , 0, false, ClearValueUint16, &FRHICommandListImmediate::ClearUAVUint, { 0x23, 0x01 }));
-			RUN_TEST(RunTest_UAVClear_Buffer(RHICmdList, BUF_VertexBuffer, 512, PF_R16G16_UINT      , 0, false, ClearValueUint16, &FRHICommandListImmediate::ClearUAVUint, { 0x23, 0x01, 0x67, 0x45 }));
-			RUN_TEST(RunTest_UAVClear_Buffer(RHICmdList, BUF_VertexBuffer, 512, PF_R16G16B16A16_UINT, 0, false, ClearValueUint16, &FRHICommandListImmediate::ClearUAVUint, { 0x23, 0x01, 0x67, 0x45, 0xab, 0x89, 0xef, 0xcd }));
-			RUN_TEST(RunTest_UAVClear_Buffer(RHICmdList, BUF_VertexBuffer, 512, PF_R32_UINT         , 0, false, ClearValueUint32, &FRHICommandListImmediate::ClearUAVUint, { 0x67, 0x45, 0x23, 0x01 }));
-			RUN_TEST(RunTest_UAVClear_Buffer(RHICmdList, BUF_VertexBuffer, 512, PF_R32G32_UINT      , 0, false, ClearValueUint32, &FRHICommandListImmediate::ClearUAVUint, { 0x67, 0x45, 0x23, 0x01, 0xef, 0xcd, 0xab, 0x89 }));
-			RUN_TEST(RunTest_UAVClear_Buffer(RHICmdList, BUF_VertexBuffer, 512, PF_R32G32B32A32_UINT, 0, false, ClearValueUint32, &FRHICommandListImmediate::ClearUAVUint, { 0x67, 0x45, 0x23, 0x01, 0xef, 0xcd, 0xab, 0x89, 0xbb, 0xaa, 0x99, 0x88, 0xff, 0xee, 0xdd, 0xcc }));
+			RUN_TEST(RunTest_UAVClear_Buffer(RHICmdList, BUF_VertexBuffer, 1024, PF_R8_UINT          , 0, false, ClearValueUint8 , &FRHICommandListImmediate::ClearUAVUint, { 0x01 }));
+			RUN_TEST(RunTest_UAVClear_Buffer(RHICmdList, BUF_VertexBuffer, 1024, PF_R8G8B8A8_UINT    , 0, false, ClearValueUint8 , &FRHICommandListImmediate::ClearUAVUint, { 0x01, 0x23, 0x45, 0x67 }));
+			RUN_TEST(RunTest_UAVClear_Buffer(RHICmdList, BUF_VertexBuffer, 1024, PF_R16_UINT         , 0, false, ClearValueUint16, &FRHICommandListImmediate::ClearUAVUint, { 0x23, 0x01 }));
+			RUN_TEST(RunTest_UAVClear_Buffer(RHICmdList, BUF_VertexBuffer, 1024, PF_R16G16_UINT      , 0, false, ClearValueUint16, &FRHICommandListImmediate::ClearUAVUint, { 0x23, 0x01, 0x67, 0x45 }));
+			RUN_TEST(RunTest_UAVClear_Buffer(RHICmdList, BUF_VertexBuffer, 1024, PF_R16G16B16A16_UINT, 0, false, ClearValueUint16, &FRHICommandListImmediate::ClearUAVUint, { 0x23, 0x01, 0x67, 0x45, 0xab, 0x89, 0xef, 0xcd }));
+			RUN_TEST(RunTest_UAVClear_Buffer(RHICmdList, BUF_VertexBuffer, 1024, PF_R32_UINT         , 0, false, ClearValueUint32, &FRHICommandListImmediate::ClearUAVUint, { 0x67, 0x45, 0x23, 0x01 }));
+			RUN_TEST(RunTest_UAVClear_Buffer(RHICmdList, BUF_VertexBuffer, 1024, PF_R32G32_UINT      , 0, false, ClearValueUint32, &FRHICommandListImmediate::ClearUAVUint, { 0x67, 0x45, 0x23, 0x01, 0xef, 0xcd, 0xab, 0x89 }));
+			RUN_TEST(RunTest_UAVClear_Buffer(RHICmdList, BUF_VertexBuffer, 1024, PF_R32G32B32A32_UINT, 0, false, ClearValueUint32, &FRHICommandListImmediate::ClearUAVUint, { 0x67, 0x45, 0x23, 0x01, 0xef, 0xcd, 0xab, 0x89, 0xbb, 0xaa, 0x99, 0x88, 0xff, 0xee, 0xdd, 0xcc }));
 
 			// Signed integer
 			const FUintVector4 ClearValueInt16_Positive(0x1122, 0x3344, 0x5566, 0x7788);
 			const FUintVector4 ClearValueInt32_Positive(0x10112233, 0x44556677, 0x0899aabb, 0x4cddeeff);
-			RUN_TEST(RunTest_UAVClear_Buffer(RHICmdList, BUF_VertexBuffer, 512, PF_R16_SINT         , 0, false, ClearValueInt16_Positive, &FRHICommandListImmediate::ClearUAVUint, { 0x22, 0x11 }));
-			RUN_TEST(RunTest_UAVClear_Buffer(RHICmdList, BUF_VertexBuffer, 512, PF_R16G16B16A16_SINT, 0, false, ClearValueInt16_Positive, &FRHICommandListImmediate::ClearUAVUint, { 0x22, 0x11, 0x44, 0x33, 0x66, 0x55, 0x88, 0x77 }));
-			RUN_TEST(RunTest_UAVClear_Buffer(RHICmdList, BUF_VertexBuffer, 512, PF_R32_SINT         , 0, false, ClearValueInt32_Positive, &FRHICommandListImmediate::ClearUAVUint, { 0x33, 0x22, 0x11, 0x10 }));
+			RUN_TEST(RunTest_UAVClear_Buffer(RHICmdList, BUF_VertexBuffer, 1024, PF_R16_SINT         , 0, false, ClearValueInt16_Positive, &FRHICommandListImmediate::ClearUAVUint, { 0x22, 0x11 }));
+			RUN_TEST(RunTest_UAVClear_Buffer(RHICmdList, BUF_VertexBuffer, 1024, PF_R16G16B16A16_SINT, 0, false, ClearValueInt16_Positive, &FRHICommandListImmediate::ClearUAVUint, { 0x22, 0x11, 0x44, 0x33, 0x66, 0x55, 0x88, 0x77 }));
+			RUN_TEST(RunTest_UAVClear_Buffer(RHICmdList, BUF_VertexBuffer, 1024, PF_R32_SINT         , 0, false, ClearValueInt32_Positive, &FRHICommandListImmediate::ClearUAVUint, { 0x33, 0x22, 0x11, 0x10 }));
 
 			const FUintVector4 ClearValueInt16_Negative(0xffff9122, 0xffffb344, 0xffffd566, 0xfffff788);
 			const FUintVector4 ClearValueInt32_Negative(0x80112233, 0xc4556677, 0x8899aabb, 0xccddeeff);
-			RUN_TEST(RunTest_UAVClear_Buffer(RHICmdList, BUF_VertexBuffer, 512, PF_R16_SINT         , 0, false, ClearValueInt16_Negative, &FRHICommandListImmediate::ClearUAVUint, { 0x22, 0x91 }));
-			RUN_TEST(RunTest_UAVClear_Buffer(RHICmdList, BUF_VertexBuffer, 512, PF_R16G16B16A16_SINT, 0, false, ClearValueInt16_Negative, &FRHICommandListImmediate::ClearUAVUint, { 0x22, 0x91, 0x44, 0xb3, 0x66, 0xd5, 0x88, 0xf7 }));
-			RUN_TEST(RunTest_UAVClear_Buffer(RHICmdList, BUF_VertexBuffer, 512, PF_R32_SINT         , 0, false, ClearValueInt32_Negative, &FRHICommandListImmediate::ClearUAVUint, { 0x33, 0x22, 0x11, 0x80 }));
+			RUN_TEST(RunTest_UAVClear_Buffer(RHICmdList, BUF_VertexBuffer, 1024, PF_R16_SINT         , 0, false, ClearValueInt16_Negative, &FRHICommandListImmediate::ClearUAVUint, { 0x22, 0x91 }));
+			RUN_TEST(RunTest_UAVClear_Buffer(RHICmdList, BUF_VertexBuffer, 1024, PF_R16G16B16A16_SINT, 0, false, ClearValueInt16_Negative, &FRHICommandListImmediate::ClearUAVUint, { 0x22, 0x91, 0x44, 0xb3, 0x66, 0xd5, 0x88, 0xf7 }));
+			RUN_TEST(RunTest_UAVClear_Buffer(RHICmdList, BUF_VertexBuffer, 1024, PF_R32_SINT         , 0, false, ClearValueInt32_Negative, &FRHICommandListImmediate::ClearUAVUint, { 0x33, 0x22, 0x11, 0x80 }));
 
 			// Raw (ByteAddress) buffer. A raw view uses only the first channle of the clear value, directly as a 32bit integer. Typed views still do the appropriate format conversion.
-			RUN_TEST(RunTest_UAVClear_Buffer(RHICmdList, BUF_VertexBuffer | BUF_ByteAddressBuffer, 512, PF_R8G8B8A8_UINT, 0, false, ClearValueUint8 , &FRHICommandListImmediate::ClearUAVUint, { 0x01, 0x23, 0x45, 0x67 }));
-			RUN_TEST(RunTest_UAVClear_Buffer(RHICmdList, BUF_VertexBuffer | BUF_ByteAddressBuffer, 512, PF_R8G8B8A8_UINT, 0, true , ClearValueUint8 , &FRHICommandListImmediate::ClearUAVUint, { 0x01, 0x00, 0x00, 0x00 }));
-			RUN_TEST(RunTest_UAVClear_Buffer(RHICmdList, BUF_VertexBuffer | BUF_ByteAddressBuffer, 512, PF_R16G16_UINT  , 0, false, ClearValueUint16, &FRHICommandListImmediate::ClearUAVUint, { 0x23, 0x01, 0x67, 0x45 }));
-			RUN_TEST(RunTest_UAVClear_Buffer(RHICmdList, BUF_VertexBuffer | BUF_ByteAddressBuffer, 512, PF_R16G16_UINT  , 0, true , ClearValueUint16, &FRHICommandListImmediate::ClearUAVUint, { 0x23, 0x01, 0x00, 0x00 }));
-			RUN_TEST(RunTest_UAVClear_Buffer(RHICmdList, BUF_VertexBuffer | BUF_ByteAddressBuffer, 512, PF_R32_UINT     , 0, false, ClearValueUint32, &FRHICommandListImmediate::ClearUAVUint, { 0x67, 0x45, 0x23, 0x01 }));
-			RUN_TEST(RunTest_UAVClear_Buffer(RHICmdList, BUF_VertexBuffer | BUF_ByteAddressBuffer, 512, PF_R32_UINT     , 0, true , ClearValueUint32, &FRHICommandListImmediate::ClearUAVUint, { 0x67, 0x45, 0x23, 0x01 }));
+			RUN_TEST(RunTest_UAVClear_Buffer(RHICmdList, BUF_VertexBuffer | BUF_ByteAddressBuffer, 1024, PF_R8G8B8A8_UINT, 0, false, ClearValueUint8 , &FRHICommandListImmediate::ClearUAVUint, { 0x01, 0x23, 0x45, 0x67 }));
+			RUN_TEST(RunTest_UAVClear_Buffer(RHICmdList, BUF_VertexBuffer | BUF_ByteAddressBuffer, 1024, PF_R8G8B8A8_UINT, 0, true , ClearValueUint8 , &FRHICommandListImmediate::ClearUAVUint, { 0x01, 0x00, 0x00, 0x00 }));
+			RUN_TEST(RunTest_UAVClear_Buffer(RHICmdList, BUF_VertexBuffer | BUF_ByteAddressBuffer, 1024, PF_R16G16_UINT  , 0, false, ClearValueUint16, &FRHICommandListImmediate::ClearUAVUint, { 0x23, 0x01, 0x67, 0x45 }));
+			RUN_TEST(RunTest_UAVClear_Buffer(RHICmdList, BUF_VertexBuffer | BUF_ByteAddressBuffer, 1024, PF_R16G16_UINT  , 0, true , ClearValueUint16, &FRHICommandListImmediate::ClearUAVUint, { 0x23, 0x01, 0x00, 0x00 }));
+			RUN_TEST(RunTest_UAVClear_Buffer(RHICmdList, BUF_VertexBuffer | BUF_ByteAddressBuffer, 1024, PF_R32_UINT     , 0, false, ClearValueUint32, &FRHICommandListImmediate::ClearUAVUint, { 0x67, 0x45, 0x23, 0x01 }));
+			RUN_TEST(RunTest_UAVClear_Buffer(RHICmdList, BUF_VertexBuffer | BUF_ByteAddressBuffer, 1024, PF_R32_UINT     , 0, true , ClearValueUint32, &FRHICommandListImmediate::ClearUAVUint, { 0x67, 0x45, 0x23, 0x01 }));
 		}
 
 		{
 			// Clamping unsigned int tests (components of ClearValueUint are > 0xffff, so will be clamped by the format conversion for formats < 32 bits per channel wide).
 			const FUintVector4 ClearValueUint(0xeeffccdd, 0xaabb8899, 0x66774455, 0x22330011);
-			RUN_TEST(RunTest_UAVClear_Buffer(RHICmdList, BUF_VertexBuffer, 512, PF_R8_UINT          , 0, false, ClearValueUint, &FRHICommandListImmediate::ClearUAVUint, { 0xff }));
-			RUN_TEST(RunTest_UAVClear_Buffer(RHICmdList, BUF_VertexBuffer, 512, PF_R16_UINT         , 0, false, ClearValueUint, &FRHICommandListImmediate::ClearUAVUint, { 0xff, 0xff }));
-			RUN_TEST(RunTest_UAVClear_Buffer(RHICmdList, BUF_VertexBuffer, 512, PF_R16G16_UINT      , 0, false, ClearValueUint, &FRHICommandListImmediate::ClearUAVUint, { 0xff, 0xff, 0xff, 0xff }));
-			RUN_TEST(RunTest_UAVClear_Buffer(RHICmdList, BUF_VertexBuffer, 512, PF_R16G16B16A16_UINT, 0, false, ClearValueUint, &FRHICommandListImmediate::ClearUAVUint, { 0xff, 0xff, 0xff, 0xff, 0xff, 0xff, 0xff, 0xff }));
-			RUN_TEST(RunTest_UAVClear_Buffer(RHICmdList, BUF_VertexBuffer, 512, PF_R8G8B8A8_UINT    , 0, false, ClearValueUint, &FRHICommandListImmediate::ClearUAVUint, { 0xff, 0xff, 0xff, 0xff }));
-			RUN_TEST(RunTest_UAVClear_Buffer(RHICmdList, BUF_VertexBuffer, 512, PF_R32_UINT         , 0, false, ClearValueUint, &FRHICommandListImmediate::ClearUAVUint, { 0xdd, 0xcc, 0xff, 0xee }));
-			RUN_TEST(RunTest_UAVClear_Buffer(RHICmdList, BUF_VertexBuffer, 512, PF_R32G32_UINT      , 0, false, ClearValueUint, &FRHICommandListImmediate::ClearUAVUint, { 0xdd, 0xcc, 0xff, 0xee, 0x99, 0x88, 0xbb, 0xaa }));
-			RUN_TEST(RunTest_UAVClear_Buffer(RHICmdList, BUF_VertexBuffer, 512, PF_R32G32B32A32_UINT, 0, false, ClearValueUint, &FRHICommandListImmediate::ClearUAVUint, { 0xdd, 0xcc, 0xff, 0xee, 0x99, 0x88, 0xbb, 0xaa, 0x55, 0x44, 0x77, 0x66, 0x11, 0x00, 0x33, 0x22 }));
+			RUN_TEST(RunTest_UAVClear_Buffer(RHICmdList, BUF_VertexBuffer, 1024, PF_R8_UINT          , 0, false, ClearValueUint, &FRHICommandListImmediate::ClearUAVUint, { 0xff }));
+			RUN_TEST(RunTest_UAVClear_Buffer(RHICmdList, BUF_VertexBuffer, 1024, PF_R16_UINT         , 0, false, ClearValueUint, &FRHICommandListImmediate::ClearUAVUint, { 0xff, 0xff }));
+			RUN_TEST(RunTest_UAVClear_Buffer(RHICmdList, BUF_VertexBuffer, 1024, PF_R16G16_UINT      , 0, false, ClearValueUint, &FRHICommandListImmediate::ClearUAVUint, { 0xff, 0xff, 0xff, 0xff }));
+			RUN_TEST(RunTest_UAVClear_Buffer(RHICmdList, BUF_VertexBuffer, 1024, PF_R16G16B16A16_UINT, 0, false, ClearValueUint, &FRHICommandListImmediate::ClearUAVUint, { 0xff, 0xff, 0xff, 0xff, 0xff, 0xff, 0xff, 0xff }));
+			RUN_TEST(RunTest_UAVClear_Buffer(RHICmdList, BUF_VertexBuffer, 1024, PF_R8G8B8A8_UINT    , 0, false, ClearValueUint, &FRHICommandListImmediate::ClearUAVUint, { 0xff, 0xff, 0xff, 0xff }));
+			RUN_TEST(RunTest_UAVClear_Buffer(RHICmdList, BUF_VertexBuffer, 1024, PF_R32_UINT         , 0, false, ClearValueUint, &FRHICommandListImmediate::ClearUAVUint, { 0xdd, 0xcc, 0xff, 0xee }));
+			RUN_TEST(RunTest_UAVClear_Buffer(RHICmdList, BUF_VertexBuffer, 1024, PF_R32G32_UINT      , 0, false, ClearValueUint, &FRHICommandListImmediate::ClearUAVUint, { 0xdd, 0xcc, 0xff, 0xee, 0x99, 0x88, 0xbb, 0xaa }));
+			RUN_TEST(RunTest_UAVClear_Buffer(RHICmdList, BUF_VertexBuffer, 1024, PF_R32G32B32A32_UINT, 0, false, ClearValueUint, &FRHICommandListImmediate::ClearUAVUint, { 0xdd, 0xcc, 0xff, 0xee, 0x99, 0x88, 0xbb, 0xaa, 0x55, 0x44, 0x77, 0x66, 0x11, 0x00, 0x33, 0x22 }));
 
 			// Signed integer
 			const FUintVector4 ClearValueInt16_ClampToMaxInt16(0x8001, 0x8233, 0x8455, 0x8677);
 			const FUintVector4 ClearValueInt16_ClampToMinInt16(0xfabc7123, 0x80123456, 0x80203040, 0x8a0b0c0d);
-			RUN_TEST(RunTest_UAVClear_Buffer(RHICmdList, BUF_VertexBuffer, 512, PF_R16_SINT         , 0, false, ClearValueInt16_ClampToMaxInt16, &FRHICommandListImmediate::ClearUAVUint, { 0xff, 0x7f }));
-			RUN_TEST(RunTest_UAVClear_Buffer(RHICmdList, BUF_VertexBuffer, 512, PF_R16_SINT         , 0, false, ClearValueInt16_ClampToMinInt16, &FRHICommandListImmediate::ClearUAVUint, { 0x00, 0x80 }));
-			RUN_TEST(RunTest_UAVClear_Buffer(RHICmdList, BUF_VertexBuffer, 512, PF_R16G16B16A16_SINT, 0, false, ClearValueInt16_ClampToMaxInt16, &FRHICommandListImmediate::ClearUAVUint, { 0xff, 0x7f, 0xff, 0x7f, 0xff, 0x7f, 0xff, 0x7f }));
-			RUN_TEST(RunTest_UAVClear_Buffer(RHICmdList, BUF_VertexBuffer, 512, PF_R16G16B16A16_SINT, 0, false, ClearValueInt16_ClampToMinInt16, &FRHICommandListImmediate::ClearUAVUint, { 0x00, 0x80, 0x00, 0x80, 0x00, 0x80, 0x00, 0x80 }));
-			RUN_TEST(RunTest_UAVClear_Buffer(RHICmdList, BUF_VertexBuffer, 512, PF_R32_SINT         , 0, false, ClearValueUint                 , &FRHICommandListImmediate::ClearUAVUint, { 0xdd, 0xcc, 0xff, 0xee }));
+			RUN_TEST(RunTest_UAVClear_Buffer(RHICmdList, BUF_VertexBuffer, 1024, PF_R16_SINT         , 0, false, ClearValueInt16_ClampToMaxInt16, &FRHICommandListImmediate::ClearUAVUint, { 0xff, 0x7f }));
+			RUN_TEST(RunTest_UAVClear_Buffer(RHICmdList, BUF_VertexBuffer, 1024, PF_R16_SINT         , 0, false, ClearValueInt16_ClampToMinInt16, &FRHICommandListImmediate::ClearUAVUint, { 0x00, 0x80 }));
+			RUN_TEST(RunTest_UAVClear_Buffer(RHICmdList, BUF_VertexBuffer, 1024, PF_R16G16B16A16_SINT, 0, false, ClearValueInt16_ClampToMaxInt16, &FRHICommandListImmediate::ClearUAVUint, { 0xff, 0x7f, 0xff, 0x7f, 0xff, 0x7f, 0xff, 0x7f }));
+			RUN_TEST(RunTest_UAVClear_Buffer(RHICmdList, BUF_VertexBuffer, 1024, PF_R16G16B16A16_SINT, 0, false, ClearValueInt16_ClampToMinInt16, &FRHICommandListImmediate::ClearUAVUint, { 0x00, 0x80, 0x00, 0x80, 0x00, 0x80, 0x00, 0x80 }));
+			RUN_TEST(RunTest_UAVClear_Buffer(RHICmdList, BUF_VertexBuffer, 1024, PF_R32_SINT         , 0, false, ClearValueUint                 , &FRHICommandListImmediate::ClearUAVUint, { 0xdd, 0xcc, 0xff, 0xee }));
 
 			// Raw (ByteAddress) buffer. A raw view uses only the first channle of the clear value, directly as a 32bit integer. Typed views still do the appropriate format conversion.
-			RUN_TEST(RunTest_UAVClear_Buffer(RHICmdList, BUF_VertexBuffer | BUF_ByteAddressBuffer, 512, PF_R16G16_UINT  , 0, false, ClearValueUint, &FRHICommandListImmediate::ClearUAVUint, { 0xff, 0xff, 0xff, 0xff }));
-			RUN_TEST(RunTest_UAVClear_Buffer(RHICmdList, BUF_VertexBuffer | BUF_ByteAddressBuffer, 512, PF_R16G16_UINT  , 0, true , ClearValueUint, &FRHICommandListImmediate::ClearUAVUint, { 0xdd, 0xcc, 0xff, 0xee }));
-			RUN_TEST(RunTest_UAVClear_Buffer(RHICmdList, BUF_VertexBuffer | BUF_ByteAddressBuffer, 512, PF_R8G8B8A8_UINT, 0, false, ClearValueUint, &FRHICommandListImmediate::ClearUAVUint, { 0xff, 0xff, 0xff, 0xff }));
-			RUN_TEST(RunTest_UAVClear_Buffer(RHICmdList, BUF_VertexBuffer | BUF_ByteAddressBuffer, 512, PF_R8G8B8A8_UINT, 0, true , ClearValueUint, &FRHICommandListImmediate::ClearUAVUint, { 0xdd, 0xcc, 0xff, 0xee }));
-			RUN_TEST(RunTest_UAVClear_Buffer(RHICmdList, BUF_VertexBuffer | BUF_ByteAddressBuffer, 512, PF_R32_UINT     , 0, false, ClearValueUint, &FRHICommandListImmediate::ClearUAVUint, { 0xdd, 0xcc, 0xff, 0xee }));
-			RUN_TEST(RunTest_UAVClear_Buffer(RHICmdList, BUF_VertexBuffer | BUF_ByteAddressBuffer, 512, PF_R32_UINT     , 0, true , ClearValueUint, &FRHICommandListImmediate::ClearUAVUint, { 0xdd, 0xcc, 0xff, 0xee }));
+			RUN_TEST(RunTest_UAVClear_Buffer(RHICmdList, BUF_VertexBuffer | BUF_ByteAddressBuffer, 1024, PF_R16G16_UINT  , 0, false, ClearValueUint, &FRHICommandListImmediate::ClearUAVUint, { 0xff, 0xff, 0xff, 0xff }));
+			RUN_TEST(RunTest_UAVClear_Buffer(RHICmdList, BUF_VertexBuffer | BUF_ByteAddressBuffer, 1024, PF_R16G16_UINT  , 0, true , ClearValueUint, &FRHICommandListImmediate::ClearUAVUint, { 0xdd, 0xcc, 0xff, 0xee }));
+			RUN_TEST(RunTest_UAVClear_Buffer(RHICmdList, BUF_VertexBuffer | BUF_ByteAddressBuffer, 1024, PF_R8G8B8A8_UINT, 0, false, ClearValueUint, &FRHICommandListImmediate::ClearUAVUint, { 0xff, 0xff, 0xff, 0xff }));
+			RUN_TEST(RunTest_UAVClear_Buffer(RHICmdList, BUF_VertexBuffer | BUF_ByteAddressBuffer, 1024, PF_R8G8B8A8_UINT, 0, true , ClearValueUint, &FRHICommandListImmediate::ClearUAVUint, { 0xdd, 0xcc, 0xff, 0xee }));
+			RUN_TEST(RunTest_UAVClear_Buffer(RHICmdList, BUF_VertexBuffer | BUF_ByteAddressBuffer, 1024, PF_R32_UINT     , 0, false, ClearValueUint, &FRHICommandListImmediate::ClearUAVUint, { 0xdd, 0xcc, 0xff, 0xee }));
+			RUN_TEST(RunTest_UAVClear_Buffer(RHICmdList, BUF_VertexBuffer | BUF_ByteAddressBuffer, 1024, PF_R32_UINT     , 0, true , ClearValueUint, &FRHICommandListImmediate::ClearUAVUint, { 0xdd, 0xcc, 0xff, 0xee }));
 		}
 
 		return bResult;
@@ -315,24 +315,24 @@ public:
 			const FVector4f ClearValueFloat(0.2345f, 0.8499f, 0.417f, 0.00145f);
 
 			// Half precision float tests
-			RUN_TEST(RunTest_UAVClear_Buffer(RHICmdList, BUF_VertexBuffer, 512, PF_R16F          , 0, false, ClearValueFloat, &FRHICommandListImmediate::ClearUAVFloat, { 0x81, 0x33 }));
-			RUN_TEST(RunTest_UAVClear_Buffer(RHICmdList, BUF_VertexBuffer, 512, PF_R16F_FILTER   , 0, false, ClearValueFloat, &FRHICommandListImmediate::ClearUAVFloat, { 0x81, 0x33 }));
-			RUN_TEST(RunTest_UAVClear_Buffer(RHICmdList, BUF_VertexBuffer, 512, PF_G16R16F       , 0, false, ClearValueFloat, &FRHICommandListImmediate::ClearUAVFloat, { 0x81, 0x33, 0xcc, 0x3a }));
-			RUN_TEST(RunTest_UAVClear_Buffer(RHICmdList, BUF_VertexBuffer, 512, PF_G16R16F_FILTER, 0, false, ClearValueFloat, &FRHICommandListImmediate::ClearUAVFloat, { 0x81, 0x33, 0xcc, 0x3a }));
+			RUN_TEST(RunTest_UAVClear_Buffer(RHICmdList, BUF_VertexBuffer, 1024, PF_R16F          , 0, false, ClearValueFloat, &FRHICommandListImmediate::ClearUAVFloat, { 0x81, 0x33 }));
+			RUN_TEST(RunTest_UAVClear_Buffer(RHICmdList, BUF_VertexBuffer, 1024, PF_R16F_FILTER   , 0, false, ClearValueFloat, &FRHICommandListImmediate::ClearUAVFloat, { 0x81, 0x33 }));
+			RUN_TEST(RunTest_UAVClear_Buffer(RHICmdList, BUF_VertexBuffer, 1024, PF_G16R16F       , 0, false, ClearValueFloat, &FRHICommandListImmediate::ClearUAVFloat, { 0x81, 0x33, 0xcc, 0x3a }));
+			RUN_TEST(RunTest_UAVClear_Buffer(RHICmdList, BUF_VertexBuffer, 1024, PF_G16R16F_FILTER, 0, false, ClearValueFloat, &FRHICommandListImmediate::ClearUAVFloat, { 0x81, 0x33, 0xcc, 0x3a }));
 
 			// Full precision float tests
-			RUN_TEST(RunTest_UAVClear_Buffer(RHICmdList, BUF_VertexBuffer, 512, PF_R32_FLOAT    , 0, false, ClearValueFloat, &FRHICommandListImmediate::ClearUAVFloat, { 0xc5, 0x20, 0x70, 0x3e }));
-			RUN_TEST(RunTest_UAVClear_Buffer(RHICmdList, BUF_VertexBuffer, 512, PF_G32R32F      , 0, false, ClearValueFloat, &FRHICommandListImmediate::ClearUAVFloat, { 0xc5, 0x20, 0x70, 0x3e, 0x0c, 0x93, 0x59, 0x3f }));
-			RUN_TEST(RunTest_UAVClear_Buffer(RHICmdList, BUF_VertexBuffer, 512, PF_A32B32G32R32F, 0, false, ClearValueFloat, &FRHICommandListImmediate::ClearUAVFloat, { 0xc5, 0x20, 0x70, 0x3e, 0x0c, 0x93, 0x59, 0x3f, 0x06, 0x81, 0xd5, 0x3e, 0xed, 0x0d, 0xbe, 0x3a }));
+			RUN_TEST(RunTest_UAVClear_Buffer(RHICmdList, BUF_VertexBuffer, 1024, PF_R32_FLOAT    , 0, false, ClearValueFloat, &FRHICommandListImmediate::ClearUAVFloat, { 0xc5, 0x20, 0x70, 0x3e }));
+			RUN_TEST(RunTest_UAVClear_Buffer(RHICmdList, BUF_VertexBuffer, 1024, PF_G32R32F      , 0, false, ClearValueFloat, &FRHICommandListImmediate::ClearUAVFloat, { 0xc5, 0x20, 0x70, 0x3e, 0x0c, 0x93, 0x59, 0x3f }));
+			RUN_TEST(RunTest_UAVClear_Buffer(RHICmdList, BUF_VertexBuffer, 1024, PF_A32B32G32R32F, 0, false, ClearValueFloat, &FRHICommandListImmediate::ClearUAVFloat, { 0xc5, 0x20, 0x70, 0x3e, 0x0c, 0x93, 0x59, 0x3f, 0x06, 0x81, 0xd5, 0x3e, 0xed, 0x0d, 0xbe, 0x3a }));
 
 			// @todo - 11,11,10 formats etc.
 
-			RUN_TEST(RunTest_UAVClear_Buffer(RHICmdList, BUF_VertexBuffer | BUF_ByteAddressBuffer, 512, PF_G16R16F       , 0, false, ClearValueFloat, &FRHICommandListImmediate::ClearUAVFloat, { 0x81, 0x33, 0xcc, 0x3a }));
-			RUN_TEST(RunTest_UAVClear_Buffer(RHICmdList, BUF_VertexBuffer | BUF_ByteAddressBuffer, 512, PF_G16R16F       , 0, true , ClearValueFloat, &FRHICommandListImmediate::ClearUAVFloat, { 0xc5, 0x20, 0x70, 0x3e }));
-			RUN_TEST(RunTest_UAVClear_Buffer(RHICmdList, BUF_VertexBuffer | BUF_ByteAddressBuffer, 512, PF_G16R16F_FILTER, 0, false, ClearValueFloat, &FRHICommandListImmediate::ClearUAVFloat, { 0x81, 0x33, 0xcc, 0x3a }));
-			RUN_TEST(RunTest_UAVClear_Buffer(RHICmdList, BUF_VertexBuffer | BUF_ByteAddressBuffer, 512, PF_G16R16F_FILTER, 0, true , ClearValueFloat, &FRHICommandListImmediate::ClearUAVFloat, { 0xc5, 0x20, 0x70, 0x3e }));
-			RUN_TEST(RunTest_UAVClear_Buffer(RHICmdList, BUF_VertexBuffer | BUF_ByteAddressBuffer, 512, PF_R32_FLOAT     , 0, false, ClearValueFloat, &FRHICommandListImmediate::ClearUAVFloat, { 0xc5, 0x20, 0x70, 0x3e }));
-			RUN_TEST(RunTest_UAVClear_Buffer(RHICmdList, BUF_VertexBuffer | BUF_ByteAddressBuffer, 512, PF_R32_FLOAT     , 0, true , ClearValueFloat, &FRHICommandListImmediate::ClearUAVFloat, { 0xc5, 0x20, 0x70, 0x3e }));
+			RUN_TEST(RunTest_UAVClear_Buffer(RHICmdList, BUF_VertexBuffer | BUF_ByteAddressBuffer, 1024, PF_G16R16F       , 0, false, ClearValueFloat, &FRHICommandListImmediate::ClearUAVFloat, { 0x81, 0x33, 0xcc, 0x3a }));
+			RUN_TEST(RunTest_UAVClear_Buffer(RHICmdList, BUF_VertexBuffer | BUF_ByteAddressBuffer, 1024, PF_G16R16F       , 0, true , ClearValueFloat, &FRHICommandListImmediate::ClearUAVFloat, { 0xc5, 0x20, 0x70, 0x3e }));
+			RUN_TEST(RunTest_UAVClear_Buffer(RHICmdList, BUF_VertexBuffer | BUF_ByteAddressBuffer, 1024, PF_G16R16F_FILTER, 0, false, ClearValueFloat, &FRHICommandListImmediate::ClearUAVFloat, { 0x81, 0x33, 0xcc, 0x3a }));
+			RUN_TEST(RunTest_UAVClear_Buffer(RHICmdList, BUF_VertexBuffer | BUF_ByteAddressBuffer, 1024, PF_G16R16F_FILTER, 0, true , ClearValueFloat, &FRHICommandListImmediate::ClearUAVFloat, { 0xc5, 0x20, 0x70, 0x3e }));
+			RUN_TEST(RunTest_UAVClear_Buffer(RHICmdList, BUF_VertexBuffer | BUF_ByteAddressBuffer, 1024, PF_R32_FLOAT     , 0, false, ClearValueFloat, &FRHICommandListImmediate::ClearUAVFloat, { 0xc5, 0x20, 0x70, 0x3e }));
+			RUN_TEST(RunTest_UAVClear_Buffer(RHICmdList, BUF_VertexBuffer | BUF_ByteAddressBuffer, 1024, PF_R32_FLOAT     , 0, true , ClearValueFloat, &FRHICommandListImmediate::ClearUAVFloat, { 0xc5, 0x20, 0x70, 0x3e }));
 		}
 
 		return bResult;
