Fixing RHI unit test fails for missing PF_R8G8B8 pixel format case

diff --git a/Engine/Plugins/Tests/RHITests/Source/RHITests/Public/RHITextureTests.h b/Engine/Plugins/Tests/RHITests/Source/RHITests/Public/RHITextureTests.h
--- a/Engine/Plugins/Tests/RHITests/Source/RHITests/Public/RHITextureTests.h
+++ b/Engine/Plugins/Tests/RHITests/Source/RHITests/Public/RHITextureTests.h
@@ -411,6 +411,7 @@ public:
 
 	typedef FDataSource< uint8,			1, true >	FDataSource8x1;
 	typedef FDataSource< uint8,			2, true >	FDataSource8x2;
+	typedef FDataSource< uint8,			3, true >	FDataSource8x3;
 	typedef FDataSource< uint8,			4, true >	FDataSource8x4;
 	typedef FDataSource< uint16,		1, true >	FDataSource16x1;
 	typedef FDataSource< uint16,		2, true >	FDataSource16x2;
@@ -745,6 +746,9 @@ public:
 			case PF_R8G8_UINT:
 				RUN_TEST(Test_UpdateTexture_Impl<FDataSource8x2UInt>(RHICmdList, Format, UpdateType));
 				break;
+			case PF_R8G8B8:
+				RUN_TEST(Test_UpdateTexture_Impl<FDataSource8x3>(RHICmdList, Format, UpdateType));
+				break;
 			case PF_B8G8R8A8:
 			case PF_R8G8B8A8:
 			case PF_A8R8G8B8:
