Make RHI "Read Surface Data" test more compatible with various RHIs.

diff --git a/Engine/Plugins/Tests/RHITests/Source/RHITests/Public/RHITextureTests.h b/Engine/Plugins/Tests/RHITests/Source/RHITests/Public/RHITextureTests.h
--- a/Engine/Plugins/Tests/RHITests/Source/RHITests/Public/RHITextureTests.h
+++ b/Engine/Plugins/Tests/RHITests/Source/RHITests/Public/RHITextureTests.h
@@ -7,6 +7,8 @@
 #include "Math/PackedVector.h"
 #include "RHIGPUReadback.h"
 
+#include "RHISurfaceDataConversion.h"
+
 //UE_DISABLE_OPTIMIZATION
 
 class FRHITextureTests
@@ -1378,47 +1380,48 @@ public:
 
 			// Read surface data
 			{
+				#define TO_FRHIR10G10B10A2(r, g, b, a) {(uint32)((r * 1023.0f) + 0.5f), (uint32)((g * 1023.0f) + 0.5f), (uint32)((b * 1023.0f) + 0.5f), (uint32)((a * 3.0f) + 0.5f)}
 				const TCHAR* TestName = TEXT("Read Surface Data");
 				const FIntPoint Extent(16, 16);
 				const FIntRect ReadRegion(5, 3, 13, 11);
 				const int32 PaletteSize = 30;
-				const FVector4f Palette[PaletteSize] = {
-					FVector4f(0.2f, 0.4f, 0.6f, 1.0f),
-					FVector4f(0.6f, 0.8f, 1.0f, 1.0f),
-					FVector4f(0.6f, 0.6f, 0.2f, 1.0f),
-					FVector4f(0.4f, 0.4f, 0.6f, 1.0f),
-					FVector4f(0.8f, 0.6f, 0.2f, 1.0f),
-					FVector4f(0.0f, 0.4f, 0.4f, 1.0f),
-					FVector4f(0.2f, 0.6f, 1.0f, 1.0f),
-					FVector4f(0.6f, 0.2f, 0.0f, 1.0f),
-					FVector4f(0.8f, 0.8f, 0.6f, 1.0f),
-					FVector4f(0.4f, 0.4f, 0.4f, 1.0f),
-					FVector4f(1.0f, 0.8f, 0.4f, 1.0f),
-					FVector4f(0.4f, 0.6f, 0.8f, 1.0f),
-					FVector4f(0.4f, 0.2f, 0.4f, 1.0f),
-					FVector4f(0.6f, 0.6f, 0.8f, 1.0f),
-					FVector4f(0.8f, 0.8f, 0.8f, 1.0f),
-					FVector4f(0.4f, 0.6f, 0.6f, 1.0f),
-					FVector4f(0.8f, 0.8f, 0.4f, 1.0f),
-					FVector4f(0.8f, 0.4f, 0.0f, 1.0f),
-					FVector4f(0.6f, 0.6f, 1.0f, 1.0f),
-					FVector4f(0.0f, 0.4f, 0.8f, 1.0f),
-					FVector4f(0.6f, 0.8f, 0.8f, 1.0f),
-					FVector4f(0.6f, 0.6f, 0.6f, 1.0f),
-					FVector4f(1.0f, 0.8f, 0.0f, 1.0f),
-					FVector4f(0.0f, 0.6f, 0.6f, 1.0f),
-					FVector4f(0.6f, 0.8f, 0.2f, 1.0f),
-					FVector4f(1.0f, 0.6f, 0.0f, 1.0f),
-					FVector4f(0.6f, 0.6f, 0.4f, 1.0f),
-					FVector4f(0.4f, 0.8f, 0.8f, 1.0f),
-					FVector4f(0.2f, 0.6f, 0.4f, 1.0f),
-					FVector4f(0.8f, 0.8f, 0.2f, 1.0f),
+				const FRHIR10G10B10A2 Palette[PaletteSize] = {
+					TO_FRHIR10G10B10A2(0.2f, 0.4f, 0.6f, 1.0f),
+					TO_FRHIR10G10B10A2(0.6f, 0.8f, 1.0f, 1.0f),
+					TO_FRHIR10G10B10A2(0.6f, 0.6f, 0.2f, 1.0f),
+					TO_FRHIR10G10B10A2(0.4f, 0.4f, 0.6f, 1.0f),
+					TO_FRHIR10G10B10A2(0.8f, 0.6f, 0.2f, 1.0f),
+					TO_FRHIR10G10B10A2(0.0f, 0.4f, 0.4f, 1.0f),
+					TO_FRHIR10G10B10A2(0.2f, 0.6f, 1.0f, 1.0f),
+					TO_FRHIR10G10B10A2(0.6f, 0.2f, 0.0f, 1.0f),
+					TO_FRHIR10G10B10A2(0.8f, 0.8f, 0.6f, 1.0f),
+					TO_FRHIR10G10B10A2(0.4f, 0.4f, 0.4f, 1.0f),
+					TO_FRHIR10G10B10A2(1.0f, 0.8f, 0.4f, 1.0f),
+					TO_FRHIR10G10B10A2(0.4f, 0.6f, 0.8f, 1.0f),
+					TO_FRHIR10G10B10A2(0.4f, 0.2f, 0.4f, 1.0f),
+					TO_FRHIR10G10B10A2(0.6f, 0.6f, 0.8f, 1.0f),
+					TO_FRHIR10G10B10A2(0.8f, 0.8f, 0.8f, 1.0f),
+					TO_FRHIR10G10B10A2(0.4f, 0.6f, 0.6f, 1.0f),
+					TO_FRHIR10G10B10A2(0.8f, 0.8f, 0.4f, 1.0f),
+					TO_FRHIR10G10B10A2(0.8f, 0.4f, 0.0f, 1.0f),
+					TO_FRHIR10G10B10A2(0.6f, 0.6f, 1.0f, 1.0f),
+					TO_FRHIR10G10B10A2(0.0f, 0.4f, 0.8f, 1.0f),
+					TO_FRHIR10G10B10A2(0.6f, 0.8f, 0.8f, 1.0f),
+					TO_FRHIR10G10B10A2(0.6f, 0.6f, 0.6f, 1.0f),
+					TO_FRHIR10G10B10A2(1.0f, 0.8f, 0.0f, 1.0f),
+					TO_FRHIR10G10B10A2(0.0f, 0.6f, 0.6f, 1.0f),
+					TO_FRHIR10G10B10A2(0.6f, 0.8f, 0.2f, 1.0f),
+					TO_FRHIR10G10B10A2(1.0f, 0.6f, 0.0f, 1.0f),
+					TO_FRHIR10G10B10A2(0.6f, 0.6f, 0.4f, 1.0f),
+					TO_FRHIR10G10B10A2(0.4f, 0.8f, 0.8f, 1.0f),
+					TO_FRHIR10G10B10A2(0.2f, 0.6f, 0.4f, 1.0f),
+					TO_FRHIR10G10B10A2(0.8f, 0.8f, 0.2f, 1.0f),
 				};
 
 				const FRHITextureCreateDesc SourceDesc =
 					FRHITextureCreateDesc::Create2D(TestName)
 					.SetExtent(Extent.X, Extent.Y)
-					.SetFormat(PF_A32B32G32R32F);
+					.SetFormat(PF_A2B10G10R10);
 
 				FTextureRHIRef SourceTexture = RHICreateTexture(SourceDesc);
 
@@ -1426,7 +1429,7 @@ public:
 				void* Data = RHICmdList.LockTexture2D(SourceTexture, 0, RLM_WriteOnly, Stride, false);
 				for (int32 Y = 0; Y < Extent.Y; Y++)
 				{
-					FVector4f* Ptr = (FVector4f*)((uint8*)Data + (Y * Stride));
+					FRHIR10G10B10A2* Ptr = (FRHIR10G10B10A2*)((uint8*)Data + (Y * Stride));
 					for (int32 X = 0; X < Extent.X; X++)
 					{
 						const int32 SurfaceIndex = (Y * Extent.X) + X;
@@ -1446,7 +1449,7 @@ public:
 				const auto ValidateRHIReadSurfaceData = [&]()
 				{
 					FString SampleStr;
-					TArray<TTuple<int, int, FVector4, FVector4>> InvalidSamples;
+					TArray<TTuple<int, int, FRHIR10G10B10A2, FRHIR10G10B10A2>> InvalidSamples;
 					for (int32 Y = ReadRegion.Min.Y; Y < ReadRegion.Max.Y; Y++)
 					{
 						for (int32 X = ReadRegion.Min.X; X < ReadRegion.Max.X; X++)
@@ -1454,9 +1457,9 @@ public:
 							const int32 SurfaceIndex = (Y * Extent.X) + X;
 							const int32 PaletteIndex = SurfaceIndex % PaletteSize;
 							const int32 ReadIndex = ((Y - ReadRegion.Min.Y) * ReadRegion.Width()) + (X - ReadRegion.Min.X);
-							const FVector4f Expected = Palette[PaletteIndex];
-							const FVector4f Sample = FVector4f((float)OutData[ReadIndex].R / 255.0f, (float)OutData[ReadIndex].G / 255.0f, (float)OutData[ReadIndex].B / 255.0f, (float)OutData[ReadIndex].A / 255.0f);
-							if (Sample != Expected)
+							const FRHIR10G10B10A2 Expected = Palette[PaletteIndex];
+							const FRHIR10G10B10A2 Sample = TO_FRHIR10G10B10A2(((float)OutData[ReadIndex].R / 255.0f), ((float)OutData[ReadIndex].G / 255.0f), ((float)OutData[ReadIndex].B / 255.0f), ((float)OutData[ReadIndex].A / 255.0f));
+							if (Sample.R != Expected.R || Sample.G != Expected.G || Sample.B != Expected.B || Sample.A != Expected.A)
 							{
 								SampleStr.Append(TEXT("X"));
 								InvalidSamples.Add(MakeTuple(X, Y, Sample, Expected));
@@ -1475,10 +1478,10 @@ public:
 						UE_LOG(LogRHIUnitTestCommandlet, Display, TEXT("%s results:\n%s"), TestName, *SampleStr);
 						for (const auto& InvalidSample : InvalidSamples)
 						{
-							UE_LOG(LogRHIUnitTestCommandlet, Display, TEXT("[%2d, %2d]: [%3.2f, %3.2f, %3.2f, %3.2f] != [%3.2f, %3.2f, %3.2f, %3.2f]"),
+							UE_LOG(LogRHIUnitTestCommandlet, Display, TEXT("[%2d, %2d]: [%d, %d, %d, %d] != [%d, %d, %d, %d]"),
 								InvalidSample.Get<0>(), InvalidSample.Get<1>(),
-								InvalidSample.Get<2>().X, InvalidSample.Get<2>().Y, InvalidSample.Get<2>().Z, InvalidSample.Get<2>().W,
-								InvalidSample.Get<3>().X, InvalidSample.Get<3>().Y, InvalidSample.Get<3>().Z, InvalidSample.Get<3>().W);
+								InvalidSample.Get<2>().R, InvalidSample.Get<2>().G, InvalidSample.Get<2>().B, InvalidSample.Get<2>().A,
+								InvalidSample.Get<3>().R, InvalidSample.Get<3>().G, InvalidSample.Get<3>().B, InvalidSample.Get<3>().A);
 						}
 						return false;
 					}
@@ -1491,6 +1494,7 @@ public:
 		}
 
 		return bResult;
+		#undef TO_FRHIR10G10B10A2
 	}
 	
 	static bool Test_MultipleLockTexture2D(FRHICommandListImmediate& RHICmdList)
