Fix Test_RHIFormat_WithParams not aliasing PF_DepthStencil to X24_S8

diff --git a/Engine/Plugins/Tests/RHITests/Source/RHITests/Public/RHITextureTests.h b/Engine/Plugins/Tests/RHITests/Source/RHITests/Public/RHITextureTests.h
--- a/Engine/Plugins/Tests/RHITests/Source/RHITests/Public/RHITextureTests.h
+++ b/Engine/Plugins/Tests/RHITests/Source/RHITests/Public/RHITextureTests.h
@@ -262,9 +262,14 @@ public:
 		bool bResult = true;
 		FString TestName = FString::Printf(TEXT("Test_RHIFormat (%s, %s, %s, %d)"), GPixelFormats[ResourceFormat].Name, GPixelFormats[SRVFormat].Name, GPixelFormats[UAVFormat].Name, Flags);
 		{
-			const FRHITextureCreateDesc Desc =
+			FRHITextureCreateDesc Desc =
 				FRHITextureCreateDesc::Create2D(*TestName, Width, Height, ResourceFormat)
 				.SetFlags(Flags);
+			
+			if(ResourceFormat != SRVFormat && SRVFormat != PF_Unknown)
+			{
+				Desc.AddAliasbleFormat(SRVFormat);
+			}
 
 			FTextureRHIRef Texture = RHICreateTexture(Desc);
 			bResult = (Texture != nullptr);
