RHI: Add missing EBufferUsageFlags::AccelerationStructure flag for reserved resource test.

#rnx

diff --git a/Engine/Plugins/Tests/RHITests/Source/RHITests/Private/RHIReservedResourceTests.cpp b/Engine/Plugins/Tests/RHITests/Source/RHITests/Private/RHIReservedResourceTests.cpp
--- a/Engine/Plugins/Tests/RHITests/Source/RHITests/Private/RHIReservedResourceTests.cpp
+++ b/Engine/Plugins/Tests/RHITests/Source/RHITests/Private/RHIReservedResourceTests.cpp
@@ -241,7 +241,7 @@ bool FRHIReservedResourceTests::Test_ReservedResource_CreateBuffer(FRHICommandLi
 	{
 		const FRHIBufferCreateDesc CreateDesc =
 			FRHIBufferCreateDesc::Create(TEXT("TestReservedAccelerationStructureBuffer"), 32 * 1024 * 1024, 4, EBufferUsageFlags::AccelerationStructure)
-			.SetUsage(EBufferUsageFlags::ReservedResource)
+			.AddUsage(EBufferUsageFlags::ReservedResource)
 			.SetInitialState(ERHIAccess::BVHWrite);
 		FBufferRHIRef Buffer = RHICmdList.CreateBuffer(CreateDesc);
 	}
