Validate 16-byte alignment on resource (sub-)allocation for all RHIs

diff --git a/Engine/Plugins/Tests/RHITests/Source/RHITests/Private/RHIAllocatorTests.cpp b/Engine/Plugins/Tests/RHITests/Source/RHITests/Private/RHIAllocatorTests.cpp
new file mode 100644
--- /dev/null
+++ b/Engine/Plugins/Tests/RHITests/Source/RHITests/Private/RHIAllocatorTests.cpp
@@ -0,0 +1,53 @@
+// Copyright Epic Games, Inc. All Rights Reserved.
+
+#include "RHIAllocatorTests.h"
+
+bool FRHIAllocatorTests::Test_LockBuffer16ByteAlignment(FRHICommandListImmediate& RHICmdList)
+{
+	bool bSucceeded = true;
+
+	const uint32 InAlignment = 0;
+	const uint32 ExpectedOutAlignment = 16;
+
+	const TArray<EBufferUsageFlags> InUsageFlags =
+	{
+		EBufferUsageFlags::Dynamic,
+		EBufferUsageFlags::Static,
+		EBufferUsageFlags::Volatile
+	};
+
+	const TArray<const TCHAR*> UsageFlagNames =
+	{
+		TEXT("Dynamic"),
+		TEXT("Static"),
+		TEXT("Volatile")
+	};
+
+	const int32 NumTests = InUsageFlags.Num();
+	for (int32 TestIdx = 0; TestIdx < NumTests; ++TestIdx)
+	{
+		EBufferUsageFlags Flags = InUsageFlags[TestIdx];
+		const FRHIBufferCreateDesc TestBufferCreateDesc =
+			FRHIBufferCreateDesc::Create(TEXT("TestLockBuffer16ByteAlignment"), uint32(sizeof(FVector4f) * 5), InAlignment, Flags | EBufferUsageFlags::VertexBuffer)
+			.SetInitialState(ERHIAccess::CopyDest);
+
+		const FBufferRHIRef TestBuffer = RHICmdList.CreateBuffer(TestBufferCreateDesc);
+
+		uintptr_t MappedAddr = (uintptr_t)RHICmdList.LockBuffer(TestBuffer, 0, sizeof(FVector4f) * 5, RLM_WriteOnly);
+		check(MappedAddr % ExpectedOutAlignment == 0);
+		bool bResult = (MappedAddr % ExpectedOutAlignment) == 0;
+
+		if (!bResult)
+		{
+			UE_LOG(LogRHIUnitTestCommandlet, Error, TEXT("Test failed. \"Test_LockBuffer16ByteAlignment\" (%s)"), UsageFlagNames[TestIdx]);
+		}
+		else
+		{
+			UE_LOG(LogRHIUnitTestCommandlet, Display, TEXT("Test passed. \"Test_LockBuffer16ByteAlignment\" (%s)"), UsageFlagNames[TestIdx]);
+		}
+
+		RHICmdList.UnlockBuffer(TestBuffer);
+		bSucceeded &= bResult;
+	}
+	return bSucceeded;
+}
\ No newline at end of file
diff --git a/Engine/Plugins/Tests/RHITests/Source/RHITests/Private/RHITests.spec.cpp b/Engine/Plugins/Tests/RHITests/Source/RHITests/Private/RHITests.spec.cpp
--- a/Engine/Plugins/Tests/RHITests/Source/RHITests/Private/RHITests.spec.cpp
+++ b/Engine/Plugins/Tests/RHITests/Source/RHITests/Private/RHITests.spec.cpp
@@ -11,6 +11,7 @@
 #include "RHIReadbackTests.h"
 #include "RHIReservedResourceTests.h"
 #include "RHIGraphicsUAVTests.h"
+#include "RHIAllocatorTests.h"
 
 BEGIN_DEFINE_SPEC(FAutomationRHITest, "Rendering.RHI", EAutomationTestFlags::EngineFilter | EAutomationTestFlags_ApplicationContextMask | EAutomationTestFlags::NonNullRHI)
 END_DEFINE_SPEC(FAutomationRHITest)
@@ -200,4 +201,14 @@ void FAutomationRHITest::Define()
 			TestEqual("Vixel shader UAV failed", bResult, 1);
 		});
 	});
+
+	Describe("Test RHI Allocators", [this]()
+	{
+		It("LockBuffer 16 Byte Alignment", [this]()
+		{
+			bool bResult = RunOnRenderThreadSynchronous(FRHIAllocatorTests::Test_LockBuffer16ByteAlignment);
+			TestEqual("LockBuffer 16 Byte Alignment Failed", bResult, 1);
+		});
+	});
+
 }
diff --git a/Engine/Plugins/Tests/RHITests/Source/RHITests/Private/RHITestsModule.cpp b/Engine/Plugins/Tests/RHITests/Source/RHITests/Private/RHITestsModule.cpp
--- a/Engine/Plugins/Tests/RHITests/Source/RHITests/Private/RHITestsModule.cpp
+++ b/Engine/Plugins/Tests/RHITests/Source/RHITests/Private/RHITestsModule.cpp
@@ -12,6 +12,7 @@
 #include "RHIReadbackTests.h"
 #include "RHIReservedResourceTests.h"
 #include "RHIGraphicsUAVTests.h"
+#include "RHIAllocatorTests.h"
 
 #define LOCTEXT_NAMESPACE "FRHITestsModule"
 
@@ -104,6 +105,13 @@ static bool RunTests_RenderThread(FRHICommandListImmediate& RHICmdList)
 		RUN_TEST(FRHIGraphicsUAVTests::Test_GraphicsUAV_VertexShader(RHICmdList));
 	}
 
+	// ------------------------------------------------
+	// RHI Allocators
+	// ------------------------------------------------
+	{
+		RUN_TEST(FRHIAllocatorTests::Test_LockBuffer16ByteAlignment(RHICmdList));
+	}
+
 	// @todo - add more tests
 	return bResult;
 }
diff --git a/Engine/Plugins/Tests/RHITests/Source/RHITests/Public/RHIAllocatorTests.h b/Engine/Plugins/Tests/RHITests/Source/RHITests/Public/RHIAllocatorTests.h
new file mode 100644
--- /dev/null
+++ b/Engine/Plugins/Tests/RHITests/Source/RHITests/Public/RHIAllocatorTests.h
@@ -0,0 +1,13 @@
+// Copyright Epic Games, Inc. All Rights Reserved.
+
+#pragma once
+
+#include "RHITestsCommon.h"
+
+
+class FRHIAllocatorTests
+{
+public:
+
+	static bool Test_LockBuffer16ByteAlignment(FRHICommandListImmediate& RHICmdList);
+};
\ No newline at end of file
diff --git a/Engine/Source/Runtime/D3D12RHI/Private/D3D12Allocation.cpp b/Engine/Source/Runtime/D3D12RHI/Private/D3D12Allocation.cpp
--- a/Engine/Source/Runtime/D3D12RHI/Private/D3D12Allocation.cpp
+++ b/Engine/Source/Runtime/D3D12RHI/Private/D3D12Allocation.cpp
@@ -252,7 +252,7 @@ FD3D12BuddyAllocator::FD3D12BuddyAllocator(FD3D12Device* ParentDevice,
 	// maxBlockSize / MinBlockSize should be a power of two  
 	check((MaxBlockSize / MinBlockSize) * MinBlockSize == MaxBlockSize); // Evenly dividable  
 	check(0 == ((MaxBlockSize / MinBlockSize) & ((MaxBlockSize / MinBlockSize) - 1))); // Power of two  
-
+	check(0 == (MinBlockSize % 16)); // Always 16 byte aligned
 	MaxOrder = UnitSizeToOrder(SizeToUnitSize(MaxBlockSize));
 
 	Reset();
@@ -425,6 +425,9 @@ void FD3D12BuddyAllocator::Allocate(uint32 SizeInBytes, uint32 Alignment, FD3D12
 
 	uint32 SizeToAllocate = SizeInBytes;
 
+	// Need to always align to a multiple of 16
+	Alignment = Align(Alignment, 16u);
+
 	// If the alignment doesn't match the block size
 	if (Alignment != 0 && MinBlockSize % Alignment != 0)
 	{
@@ -897,6 +900,9 @@ bool FD3D12BucketAllocator::TryAllocate(uint32 SizeInBytes, uint32 Alignment, FD
 
 	FD3D12Adapter* Adapter = GetParentDevice()->GetParentAdapter();
 
+	// Alignment needs to be a multiple of 16 bytes
+	Alignment = Align(Alignment, 16u);
+
 	// Size cannot be smaller than the requested alignment
 	SizeInBytes = FMath::Max(SizeInBytes, Alignment);
 
