Fixing RHI Unit Tests with Bindless disabled.

diff --git a/Engine/Plugins/Tests/RHITests/Source/RHITests/Private/RHIBindlessTests.cpp b/Engine/Plugins/Tests/RHITests/Source/RHITests/Private/RHIBindlessTests.cpp
--- a/Engine/Plugins/Tests/RHITests/Source/RHITests/Private/RHIBindlessTests.cpp
+++ b/Engine/Plugins/Tests/RHITests/Source/RHITests/Private/RHIBindlessTests.cpp
@@ -2,10 +2,11 @@
 
 #include "RHIBindlessTests.h"
 #include "RHIResourceUtils.h"
+#include "RHI.h"
 
 bool RHIBindlessTests::Test_ResourceCollection(FRHICommandListImmediate& RHICmdList)
 {
-	if (!GRHIGlobals.bSupportsBindless)
+	if (!GRHIGlobals.bSupportsBindless || !IsBindlessEnabledForAnyGraphics(RHIGetRuntimeBindlessConfiguration(GMaxRHIShaderPlatform)))
 	{
 		return true;
 	}
