RHITests shouldn't use implicit immediate command list.

diff --git a/Engine/Plugins/Tests/RHITests/Source/RHITests/Private/RHIClearTextureTests.cpp b/Engine/Plugins/Tests/RHITests/Source/RHITests/Private/RHIClearTextureTests.cpp
--- a/Engine/Plugins/Tests/RHITests/Source/RHITests/Private/RHIClearTextureTests.cpp
+++ b/Engine/Plugins/Tests/RHITests/Source/RHITests/Private/RHIClearTextureTests.cpp
@@ -126,7 +126,7 @@ struct FStagingTexturePool
 	void ReturnToPool(const TArray<FTestOperation>& TestOperations);
 	void FlushStagingTextures(FRHICommandListImmediate& RHICmdList, EPixelFormat PixelFormat);
 	void FlushAllStagingTextures(FRHICommandListImmediate& RHICmdList);
-	void PreallocateStagingTexture(FRHICommandListImmediate& RHICmdList, EPixelFormat PixelFormat, uint32 NumTextures);
+	void PreallocateStagingTexture(FRHICommandListBase& RHICmdList, EPixelFormat PixelFormat, uint32 NumTextures);
 
 private:
 	void Recycle(FRHICommandListImmediate& RHICmdList);
@@ -151,28 +151,28 @@ FRHITextureCreateDesc CreateTexture2DAtlasDesc(EPixelFormat PixelFormat, ERHIAcc
 	return RHITextureCreateDescAtlas;
 }
 
-FTextureRHIRef CreateTexture2DAtlasRT(FRHICommandListImmediate& RHICmdList, EPixelFormat PixelFormat, FUnorderedAccessViewRHIRef* OutUAV)
+FTextureRHIRef CreateTexture2DAtlasRT(FRHICommandListBase& RHICmdList, EPixelFormat PixelFormat, FUnorderedAccessViewRHIRef* OutUAV)
 {
-	FTextureRHIRef Texture2DAtlasRT =  RHICreateTexture(CreateTexture2DAtlasDesc(PixelFormat, ERHIAccess::UAVCompute, TexCreate_UAV | TexCreate_ShaderResource, TEXT("Texture2DAtlasRT")));
+	FTextureRHIRef Texture2DAtlasRT = RHICmdList.CreateTexture(CreateTexture2DAtlasDesc(PixelFormat, ERHIAccess::UAVCompute, TexCreate_UAV | TexCreate_ShaderResource, TEXT("Texture2DAtlasRT")));
 	*OutUAV = RHICmdList.CreateUnorderedAccessView(Texture2DAtlasRT, FRHIViewDesc::CreateTextureUAV().SetDimensionFromTexture(Texture2DAtlasRT));
 	(*OutUAV)->SetOwnerName(TEXT("Texture2DAtlasRT"));
 
 	return Texture2DAtlasRT;
 }
 
-FTextureRHIRef CreateTexture2DAtlasStaging(EPixelFormat PixelFormat)
+FTextureRHIRef CreateTexture2DAtlasStaging(FRHICommandListBase& RHICmdList, EPixelFormat PixelFormat)
 {
-	return RHICreateTexture(CreateTexture2DAtlasDesc(PixelFormat, ERHIAccess::CPURead, TexCreate_CPUReadback, TEXT("Texture2DAtlasStaging")));
+	return RHICmdList.CreateTexture(CreateTexture2DAtlasDesc(PixelFormat, ERHIAccess::CPURead, TexCreate_CPUReadback, TEXT("Texture2DAtlasStaging")));
 }
 
-void FStagingTexturePool::PreallocateStagingTexture(FRHICommandListImmediate& RHICmdList, EPixelFormat PixelFormat, uint32 NumTextures)
+void FStagingTexturePool::PreallocateStagingTexture(FRHICommandListBase& RHICmdList, EPixelFormat PixelFormat, uint32 NumTextures)
 {
 	CLEAR_TEST_SCOPED_NAMED_EVENT_F(TEXT("AllocateStaging: %d in flight"), FColor::Magenta, CurrentStagingTextures[PixelFormat].Num());
 	CurrentStagingTextures[PixelFormat].Reserve(NumTextures);
 	AvailableStagingTextures[PixelFormat].Reserve(NumTextures);
 	for (uint32 TextureIndex=0; TextureIndex < NumTextures; ++TextureIndex)
 	{
-		FTextureRHIRef Texture = CreateTexture2DAtlasStaging(PixelFormat);
+		FTextureRHIRef Texture = CreateTexture2DAtlasStaging(RHICmdList, PixelFormat);
 		CurrentStagingTextures[PixelFormat].Add(Texture);
 		AvailableStagingTextures[PixelFormat].Add(Texture);
 	}
@@ -1300,7 +1300,7 @@ static FShaderResourceViewRHIRef CreateClearTextureSRV(FRHICommandListImmediate&
 
 static bool TestDrawAtlasing(FRHICommandListImmediate& RHICmdList, FVertexDeclarationRHIRef VertexDeclarationRHI, const FRHITextureCreateDesc& RHITextureCreateDesc)
 {
-	FTextureRHIRef SourceTexture = RHICreateTexture(RHITextureCreateDesc);
+	FTextureRHIRef SourceTexture = RHICmdList.CreateTexture(RHITextureCreateDesc);
 
 	// Create the SourceTexture
 	for (uint32 MipIndex=0; MipIndex < RHITextureCreateDesc.NumMips; ++MipIndex)
@@ -1354,7 +1354,7 @@ static bool TestDrawAtlasing(FRHICommandListImmediate& RHICmdList, FVertexDeclar
 
 	FUnorderedAccessViewRHIRef TmpTexture2DAtlasUAV;
 	FTextureRHIRef TmpTexture2DAtlas = CreateTexture2DAtlasRT(RHICmdList, SourceTexture->GetDesc().Format, &TmpTexture2DAtlasUAV);
-	FTextureRHIRef DestTextureStaging = CreateTexture2DAtlasStaging(SourceTexture->GetDesc().Format);
+	FTextureRHIRef DestTextureStaging = CreateTexture2DAtlasStaging(RHICmdList, SourceTexture->GetDesc().Format);
 
 	RHICmdList.Transition(FRHITransitionInfo(DestTextureStaging, ERHIAccess::CPURead, ERHIAccess::CopyDest));
 	RHICmdList.Transition(FRHITransitionInfo(SourceTexture, ERHIAccess::CopySrc, ERHIAccess::SRVCompute));
@@ -1573,8 +1573,8 @@ bool FRHIClearTextureTests::Test_ClearTexture(FRHICommandListImmediate& RHICmdLi
 			FTextureRHIRef Texture;
 			FShaderResourceViewRHIRef TextureSRV;
 			{
-				CLEAR_TEST_SCOPED_NAMED_EVENT_TEXT(TEXT("RHICreateTexture"), FColor::Magenta);
-				Texture = RHICreateTexture(RHITextureCreateDescs[TextureIndex]);
+				CLEAR_TEST_SCOPED_NAMED_EVENT_TEXT(TEXT("CreateTexture"), FColor::Magenta);
+				Texture = RHICmdList.CreateTexture(RHITextureCreateDescs[TextureIndex]);
 				TextureSRV = CreateClearTextureSRV(RHICmdList, Texture);
 			}
 			TestContext.BatchStartTextureIndex = TextureIndex;
diff --git a/Engine/Plugins/Tests/RHITests/Source/RHITests/Private/RHIDrawTests.cpp b/Engine/Plugins/Tests/RHITests/Source/RHITests/Private/RHIDrawTests.cpp
--- a/Engine/Plugins/Tests/RHITests/Source/RHITests/Private/RHIDrawTests.cpp
+++ b/Engine/Plugins/Tests/RHITests/Source/RHITests/Private/RHIDrawTests.cpp
@@ -50,7 +50,7 @@ struct FDrawTestResources
 	{		
 		FRHITextureDesc RenderTargetTextureDesc(ETextureDimension::Texture2D, ETextureCreateFlags::RenderTargetable, PF_B8G8R8A8, FClearValueBinding(), RenderTargetSize, 1, 1, 1, 1, 0);
 		FRHITextureCreateDesc RenderTargetCreateDesc(RenderTargetTextureDesc, ERHIAccess::RTV, TEXT("DrawTest_RenderTarget"));
-		RenderTarget = RHICreateTexture(RenderTargetCreateDesc);
+		RenderTarget = RHICmdList.CreateTexture(RenderTargetCreateDesc);
 
 		FVertexDeclarationElementList VertexDeclarationElements;
 		// Position vertex element
diff --git a/Engine/Plugins/Tests/RHITests/Source/RHITests/Private/RHIGraphicsUAVTests.cpp b/Engine/Plugins/Tests/RHITests/Source/RHITests/Private/RHIGraphicsUAVTests.cpp
--- a/Engine/Plugins/Tests/RHITests/Source/RHITests/Private/RHIGraphicsUAVTests.cpp
+++ b/Engine/Plugins/Tests/RHITests/Source/RHITests/Private/RHIGraphicsUAVTests.cpp
@@ -68,7 +68,7 @@ static void Test_GraphicsUAV_Common(FRHICommandListImmediate& RHICmdList,
 
 	FRHITextureCreateDesc RenderTargetCreateDesc(RenderTargetTextureDesc,
 		ERHIAccess::RTV, TEXT("GraphicsUAVTests_RenderTarget"));
-	RenderTarget = RHICreateTexture(RenderTargetCreateDesc);
+	RenderTarget = RHICmdList.CreateTexture(RenderTargetCreateDesc);
 
 	FVertexDeclarationElementList VertexDeclarationElements;
 	VertexDeclarationElements.Add(FVertexElement(0, 0, VET_Float4, 0, 16));
diff --git a/Engine/Plugins/Tests/RHITests/Source/RHITests/Private/RHIReadbackTests.cpp b/Engine/Plugins/Tests/RHITests/Source/RHITests/Private/RHIReadbackTests.cpp
--- a/Engine/Plugins/Tests/RHITests/Source/RHITests/Private/RHIReadbackTests.cpp
+++ b/Engine/Plugins/Tests/RHITests/Source/RHITests/Private/RHIReadbackTests.cpp
@@ -248,7 +248,7 @@ bool FRHIReadbackTests::Test_TextureReadback(FRHICommandListImmediate& RHICmdLis
 			.SetFormat(PF_R32_UINT)
 			.SetInitialState(ERHIAccess::UAVCompute);
 
-		Texture2D = RHICreateTexture(Texture2DDesc);
+		Texture2D = RHICmdList.CreateTexture(Texture2DDesc);
 		auto Texture2DUAVDesc = FRHIViewDesc::CreateTextureUAV()
 			.SetDimensionFromTexture(Texture2D)
 			.SetFormat(PF_R32_UINT);
@@ -266,7 +266,7 @@ bool FRHIReadbackTests::Test_TextureReadback(FRHICommandListImmediate& RHICmdLis
 			.SetFormat(PF_R32_UINT)
 			.SetInitialState(ERHIAccess::UAVCompute);
 
-		Texture3D = RHICreateTexture(Texture3DDesc);
+		Texture3D = RHICmdList.CreateTexture(Texture3DDesc);
 		auto Texture3DUAVDesc = FRHIViewDesc::CreateTextureUAV()
 			.SetDimensionFromTexture(Texture3D)
 			.SetFormat(PF_R32_UINT);
@@ -283,7 +283,7 @@ bool FRHIReadbackTests::Test_TextureReadback(FRHICommandListImmediate& RHICmdLis
 			.SetFormat(PF_R32_UINT)
 			.SetInitialState(ERHIAccess::UAVCompute);
 
-		TextureArray = RHICreateTexture(TextureArrayDesc);
+		TextureArray = RHICmdList.CreateTexture(TextureArrayDesc);
 		auto TextureArrayUAVDesc = FRHIViewDesc::CreateTextureUAV()
 			.SetDimensionFromTexture(TextureArray)
 			.SetFormat(PF_R32_UINT);
diff --git a/Engine/Plugins/Tests/RHITests/Source/RHITests/Private/RHIReservedResourceTests.cpp b/Engine/Plugins/Tests/RHITests/Source/RHITests/Private/RHIReservedResourceTests.cpp
--- a/Engine/Plugins/Tests/RHITests/Source/RHITests/Private/RHIReservedResourceTests.cpp
+++ b/Engine/Plugins/Tests/RHITests/Source/RHITests/Private/RHIReservedResourceTests.cpp
@@ -25,7 +25,7 @@ bool FRHIReservedResourceTests::Test_ReservedResource_CreateVolumeTexture(FRHICo
 			.SetFormat(PF_A32B32G32R32F)
 			.SetInitialState(ERHIAccess::SRVCompute);
 
-		FTextureRHIRef Texture = RHICreateTexture(Desc);
+		FTextureRHIRef Texture = RHICmdList.CreateTexture(Desc);
 	}
 
 	// Small size volume texture (16MB) with immediately committed physical memory.
@@ -39,7 +39,7 @@ bool FRHIReservedResourceTests::Test_ReservedResource_CreateVolumeTexture(FRHICo
 			.SetFormat(PF_A32B32G32R32F)
 			.SetInitialState(ERHIAccess::SRVCompute);
 
-		FTextureRHIRef Texture = RHICreateTexture(Desc);
+		FTextureRHIRef Texture = RHICmdList.CreateTexture(Desc);
 	}
 
 	// Medium size volume texture (128MB) with immediately committed physical memory.
@@ -54,7 +54,7 @@ bool FRHIReservedResourceTests::Test_ReservedResource_CreateVolumeTexture(FRHICo
 			.SetFormat(PF_A32B32G32R32F)
 			.SetInitialState(ERHIAccess::SRVCompute);
 
-		FTextureRHIRef Texture = RHICreateTexture(Desc);
+		FTextureRHIRef Texture = RHICmdList.CreateTexture(Desc);
 	}
 
 	// Huge volume texture without physical memory allocation.
@@ -69,7 +69,7 @@ bool FRHIReservedResourceTests::Test_ReservedResource_CreateVolumeTexture(FRHICo
 			.SetFormat(PF_A32B32G32R32F)
 			.SetInitialState(ERHIAccess::SRVCompute);
 
-		FTextureRHIRef Texture = RHICreateTexture(Desc);
+		FTextureRHIRef Texture = RHICmdList.CreateTexture(Desc);
 	}
 
 	return true;
@@ -148,7 +148,7 @@ bool FRHIReservedResourceTests::Test_ReservedResource_CreateTexture(FRHICommandL
 				.SetFormat(Format)
 				.SetInitialState(ERHIAccess::SRVCompute);
 
-			FTextureRHIRef Texture = RHICreateTexture(Desc);
+			FTextureRHIRef Texture = RHICmdList.CreateTexture(Desc);
 		}
 	}
 
@@ -162,7 +162,7 @@ bool FRHIReservedResourceTests::Test_ReservedResource_CreateTexture(FRHICommandL
 			.SetFormat(PF_A32B32G32R32F)
 			.SetInitialState(ERHIAccess::SRVCompute);
 
-		FTextureRHIRef Texture = RHICreateTexture(Desc);
+		FTextureRHIRef Texture = RHICmdList.CreateTexture(Desc);
 	}
 
 	return true;
@@ -186,7 +186,7 @@ bool FRHIReservedResourceTests::Test_ReservedResource_CreateTextureWithMips(FRHI
 			.SetFormat(PF_R8)
 			.SetInitialState(ERHIAccess::SRVCompute);
 
-		FTextureRHIRef Texture = RHICreateTexture(Desc);
+		FTextureRHIRef Texture = RHICmdList.CreateTexture(Desc);
 	}
 
 	// Try to create a 2D texture with mips
@@ -199,7 +199,7 @@ bool FRHIReservedResourceTests::Test_ReservedResource_CreateTextureWithMips(FRHI
 			.SetFormat(PF_A32B32G32R32F)
 			.SetInitialState(ERHIAccess::SRVCompute);
 
-		FTextureRHIRef Texture = RHICreateTexture(Desc);
+		FTextureRHIRef Texture = RHICmdList.CreateTexture(Desc);
 	}
 
 	return true;
diff --git a/Engine/Plugins/Tests/RHITests/Source/RHITests/Private/RHITextureTests.cpp b/Engine/Plugins/Tests/RHITests/Source/RHITests/Private/RHITextureTests.cpp
--- a/Engine/Plugins/Tests/RHITests/Source/RHITests/Private/RHITextureTests.cpp
+++ b/Engine/Plugins/Tests/RHITests/Source/RHITests/Private/RHITextureTests.cpp
@@ -23,7 +23,7 @@ bool FRHITextureTests::VerifyTextureContents(const TCHAR* TestName, FRHICommandL
 			for (uint32 Z = 0; Z < MipDepth; ++Z)
 			{
 				{
-					FTextureRHIRef StagingTexture = RHICreateTexture(Desc);
+					FTextureRHIRef StagingTexture = RHICmdList.CreateTexture(Desc);
 
 					FRHICopyTextureInfo CopyInfo = {};
 					CopyInfo.Size = FIntVector(MipWidth, MipHeight, 1);
diff --git a/Engine/Plugins/Tests/RHITests/Source/RHITests/Public/RHITextureTests.h b/Engine/Plugins/Tests/RHITests/Source/RHITests/Public/RHITextureTests.h
--- a/Engine/Plugins/Tests/RHITests/Source/RHITests/Public/RHITextureTests.h
+++ b/Engine/Plugins/Tests/RHITests/Source/RHITests/Public/RHITextureTests.h
@@ -118,7 +118,7 @@ public:
 				.SetNumMips(NumMips)
 				.SetFlags(ETextureCreateFlags::UAV | ETextureCreateFlags::ShaderResource);
 
-			FTextureRHIRef Texture = RHICreateTexture(Desc);
+			FTextureRHIRef Texture = RHICmdList.CreateTexture(Desc);
 
 			for (uint32 Mip = 0; Mip < NumMips; ++Mip)
 			{
@@ -195,7 +195,7 @@ public:
 				.SetNumMips(NumMips)
 				.SetFlags(ETextureCreateFlags::UAV | ETextureCreateFlags::ShaderResource);
 
-			FTextureRHIRef Texture = RHICreateTexture(Desc);
+			FTextureRHIRef Texture = RHICmdList.CreateTexture(Desc);
 
 			for (uint32 Mip = 0; Mip < NumMips; ++Mip)
 			{
@@ -271,7 +271,7 @@ public:
 				Desc.AddAliasbleFormat(SRVFormat);
 			}
 
-			FTextureRHIRef Texture = RHICreateTexture(Desc);
+			FTextureRHIRef Texture = RHICmdList.CreateTexture(Desc);
 			bResult = (Texture != nullptr);
 
 			if (Texture && SRVFormat != PF_Unknown)
@@ -640,7 +640,7 @@ public:
 
 			Desc.SetFlags(ETextureCreateFlags::ShaderResource);
 
-			FTextureRHIRef Texture = RHICreateTexture(Desc);
+			FTextureRHIRef Texture = RHICmdList.CreateTexture(Desc);
 			if (Texture == nullptr)
 			{
 				UE_LOG(LogRHIUnitTestCommandlet, Display, TEXT("Test failed (couldn't create texture). \"%s\""), *TestName);
@@ -959,7 +959,7 @@ public:
 					.SetFormat(PF_R32_UINT)
 					.SetInitialState(ERHIAccess::CopySrc);
 
-				FTextureRHIRef SourceTexture = RHICreateTexture(SourceDesc);
+				FTextureRHIRef SourceTexture = RHICmdList.CreateTexture(SourceDesc);
 
 				uint32 Stride = 0;
 				void* Data = RHICmdList.LockTexture2D(SourceTexture, 0, RLM_WriteOnly, Stride, false);
@@ -972,7 +972,7 @@ public:
 					.SetFormat(PF_R32_UINT)
 					.SetInitialState(ERHIAccess::CopyDest);
 
-				FTextureRHIRef DestTexture = RHICreateTexture(DestDesc);
+				FTextureRHIRef DestTexture = RHICmdList.CreateTexture(DestDesc);
 				RHICmdList.CopyTexture(SourceTexture, DestTexture, {});
 				RHICmdList.Transition(FRHITransitionInfo(DestTexture, ERHIAccess::CopyDest, ERHIAccess::CopySrc));
 
@@ -995,7 +995,7 @@ public:
 					.SetFlags(bSourceSrgb ? ETextureCreateFlags::SRGB : ETextureCreateFlags::None)
 					.SetInitialState(ERHIAccess::CopySrc);
 
-				FTextureRHIRef SourceTexture = RHICreateTexture(SourceDesc);
+				FTextureRHIRef SourceTexture = RHICmdList.CreateTexture(SourceDesc);
 
 				uint32 Stride = 0;
 				void* Data = RHICmdList.LockTexture2D(SourceTexture, 0, RLM_WriteOnly, Stride, false);
@@ -1009,7 +1009,7 @@ public:
 					.SetFlags(bDestSrgb ? ETextureCreateFlags::SRGB : ETextureCreateFlags::None)
 					.SetInitialState(ERHIAccess::CopyDest);
 
-				FTextureRHIRef DestTexture = RHICreateTexture(DestDesc);
+				FTextureRHIRef DestTexture = RHICmdList.CreateTexture(DestDesc);
 				RHICmdList.CopyTexture(SourceTexture, DestTexture, {});
 				RHICmdList.Transition(FRHITransitionInfo(DestTexture, ERHIAccess::CopyDest, ERHIAccess::CopySrc));
 
@@ -1037,7 +1037,7 @@ public:
 					.SetFormat(PF_R32_UINT)
 					.SetInitialState(ERHIAccess::CopySrc);
 
-				FTextureRHIRef SourceTexture = RHICreateTexture(SourceDesc);
+				FTextureRHIRef SourceTexture = RHICmdList.CreateTexture(SourceDesc);
 
 				for (int32 MipIndex = 0; MipIndex < NumMips; ++MipIndex)
 				{
@@ -1057,7 +1057,7 @@ public:
 						.SetFormat(PF_R32_UINT)
 						.SetInitialState(ERHIAccess::CopySrc);
 
-					FTextureRHIRef DestTexture = RHICreateTexture(DestDesc);
+					FTextureRHIRef DestTexture = RHICmdList.CreateTexture(DestDesc);
 
 					FRHICopyTextureInfo CopyInfo;
 					CopyInfo.Size = bUseExplicitSize ? FIntVector(DestDesc.Extent.X, DestDesc.Extent.Y, 1) : FIntVector::ZeroValue;
@@ -1092,7 +1092,7 @@ public:
 					.SetFormat(PF_R32_UINT)
 					.SetInitialState(ERHIAccess::CopySrc);
 
-				FTextureRHIRef SourceTexture = RHICreateTexture(SourceDesc);
+				FTextureRHIRef SourceTexture = RHICmdList.CreateTexture(SourceDesc);
 
 				uint32 Stride = 0;
 				void* Data = RHICmdList.LockTexture2D(SourceTexture, 0, RLM_WriteOnly, Stride, false);
@@ -1110,7 +1110,7 @@ public:
 				CopyInfo.SourcePosition = FIntVector(PatternOffsetSource.X, PatternOffsetSource.Y, 0);
 				CopyInfo.DestPosition   = FIntVector(PatternOffsetDest.X, PatternOffsetDest.Y, 0);
 
-				FTextureRHIRef DestTexture = RHICreateTexture(DestDesc);
+				FTextureRHIRef DestTexture = RHICmdList.CreateTexture(DestDesc);
 				RHICmdList.CopyTexture(SourceTexture, DestTexture, CopyInfo);
 				RHICmdList.Transition(FRHITransitionInfo(DestTexture, ERHIAccess::CopyDest, ERHIAccess::CopySrc));
 
@@ -1134,7 +1134,7 @@ public:
 					.SetFormat(PF_R32_UINT)
 					.SetInitialState(ERHIAccess::CopySrc);
 
-				FTextureRHIRef SourceTexture = RHICreateTexture(SourceDesc);
+				FTextureRHIRef SourceTexture = RHICmdList.CreateTexture(SourceDesc);
 
 				uint32 Stride = 0;
 				void* Data = RHICmdList.LockTexture2D(SourceTexture, 0, RLM_WriteOnly, Stride, false);
@@ -1180,7 +1180,7 @@ public:
 					.SetFormat(PF_R32_UINT)
 					.SetInitialState(ERHIAccess::CopySrc);
 
-				FTextureRHIRef SourceTexture = RHICreateTexture(SourceDesc);
+				FTextureRHIRef SourceTexture = RHICmdList.CreateTexture(SourceDesc);
 
 				int32 ImageIndex = 1;
 
@@ -1216,7 +1216,7 @@ public:
 					.SetFormat(PF_R32_UINT)
 					.SetInitialState(ERHIAccess::CopyDest);
 
-				FTextureRHIRef DestTexture = RHICreateTexture(DestDesc);
+				FTextureRHIRef DestTexture = RHICmdList.CreateTexture(DestDesc);
 
 				/*   SLICES
 				  M  13  9  5  1
@@ -1269,7 +1269,7 @@ public:
 					.SetFormat(PF_R32_UINT)
 					.SetInitialState(ERHIAccess::CopySrc);
 
-				FTextureRHIRef SourceTexture = RHICreateTexture(SourceDesc);
+				FTextureRHIRef SourceTexture = RHICmdList.CreateTexture(SourceDesc);
 
 				int32 ImageIndex = 1;
 
@@ -1298,7 +1298,7 @@ public:
 					.SetFormat(PF_R32_UINT)
 					.SetInitialState(ERHIAccess::CopyDest);
 
-				FTextureRHIRef DestTexture = RHICreateTexture(DestDesc);
+				FTextureRHIRef DestTexture = RHICmdList.CreateTexture(DestDesc);
 
 				for (int32 MipIndex = 0; MipIndex < NumMips; ++MipIndex)
 				{
@@ -1342,8 +1342,8 @@ public:
 				const FRHITextureCreateDesc DestDesc = FRHITextureCreateDesc(SourceDesc)
 					.SetInitialState(ERHIAccess::CopyDest);
 
-				FTextureRHIRef SourceTexture = RHICreateTexture(SourceDesc);
-				FTextureRHIRef DestTexture = RHICreateTexture(DestDesc);
+				FTextureRHIRef SourceTexture = RHICmdList.CreateTexture(SourceDesc);
+				FTextureRHIRef DestTexture = RHICmdList.CreateTexture(DestDesc);
 
 				// Fill the source texture with data
 				for (int32 Mip = 0, ImageIndex = 1; Mip < SourceDesc.NumMips; ++Mip, ImageIndex += SourceDesc.Depth)
@@ -1428,7 +1428,7 @@ public:
 					.SetExtent(Extent.X, Extent.Y)
 					.SetFormat(PF_A2B10G10R10);
 
-				FTextureRHIRef SourceTexture = RHICreateTexture(SourceDesc);
+				FTextureRHIRef SourceTexture = RHICmdList.CreateTexture(SourceDesc);
 
 				uint32 Stride = 0;
 				void* Data = RHICmdList.LockTexture2D(SourceTexture, 0, RLM_WriteOnly, Stride, false);
@@ -1513,8 +1513,8 @@ public:
 		const EPixelFormat Format = PF_R8G8B8A8;
 		FRHITextureCreateDesc DescArray = FRHITextureCreateDesc::Create2DArray(TEXT("Multiple Texture Slice Lock"), SliceDim, ArrayLength, Format).DetermineInititialState();
 		
-		FTextureRHIRef Texture_SingleLock = RHICreateTexture(DescArray);
-		FTextureRHIRef Texture_MultipleLock = RHICreateTexture(DescArray);
+		FTextureRHIRef Texture_SingleLock = RHICmdList.CreateTexture(DescArray);
+		FTextureRHIRef Texture_MultipleLock = RHICmdList.CreateTexture(DescArray);
 		
 		// Create Reasonable Test Data that can be seen in the GPU debugger if required
 		const auto WriteSliceTestData = [=](uint8* pDestData, uint32 DestStride, uint32 SliceIdx)
@@ -1593,8 +1593,8 @@ public:
 			.SetFlags(ETextureCreateFlags::CPUReadback)
 			.SetInitialState(ERHIAccess::CopyDest);
 
-		FTextureRHIRef SingleLockStagingTexture = RHICreateTexture(Desc);
-		FTextureRHIRef MultiLockStagingTexture = RHICreateTexture(Desc);
+		FTextureRHIRef SingleLockStagingTexture = RHICmdList.CreateTexture(Desc);
+		FTextureRHIRef MultiLockStagingTexture = RHICmdList.CreateTexture(Desc);
 
 		FRHICopyTextureInfo CopyInfo = {};
 		CopyInfo.Size = FIntVector(SliceDim.X, SliceDim.Y, 1);
