Check if hardware ray tracing is supported by the current RHI before trying to create ray tracing buffer types in the RHI unit test.

#Synced-CL 49292701

diff --git a/Engine/Plugins/Tests/RHITests/Source/RHITests/Private/RHIReservedResourceTests.cpp b/Engine/Plugins/Tests/RHITests/Source/RHITests/Private/RHIReservedResourceTests.cpp
--- a/Engine/Plugins/Tests/RHITests/Source/RHITests/Private/RHIReservedResourceTests.cpp
+++ b/Engine/Plugins/Tests/RHITests/Source/RHITests/Private/RHIReservedResourceTests.cpp
@@ -238,6 +238,7 @@ bool FRHIReservedResourceTests::Test_ReservedResource_CreateBuffer(FRHICommandLi
 		FBufferRHIRef Buffer = RHICmdList.CreateBuffer(CreateDesc);
 	}
 
+	if (GRHIGlobals.RayTracing.Supported)
 	{
 		const FRHIBufferCreateDesc CreateDesc =
 			FRHIBufferCreateDesc::Create(TEXT("TestReservedAccelerationStructureBuffer"), 32 * 1024 * 1024, 4, EBufferUsageFlags::AccelerationStructure)
@@ -246,6 +247,7 @@ bool FRHIReservedResourceTests::Test_ReservedResource_CreateBuffer(FRHICommandLi
 		FBufferRHIRef Buffer = RHICmdList.CreateBuffer(CreateDesc);
 	}
 
+	if (GRHIGlobals.RayTracing.Supported)
 	{
 		const FRHIBufferCreateDesc CreateDesc =
 			FRHIBufferCreateDesc::Create(TEXT("TestReservedRayTracingScratchBuffer"), 32 * 1024 * 1024, 4, EBufferUsageFlags::RayTracingScratch)
