Fixing "Read Surface Data" d3d validation errors by making the texture non-static. The assumption around state tracking is that we only care about transitions when something is non-static, so the transitions to CopySrc were being skipped.

#Synced-CL 49623085

diff --git a/Engine/Plugins/Tests/RHITests/Source/RHITests/Public/RHITextureTests.h b/Engine/Plugins/Tests/RHITests/Source/RHITests/Public/RHITextureTests.h
--- a/Engine/Plugins/Tests/RHITests/Source/RHITests/Public/RHITextureTests.h
+++ b/Engine/Plugins/Tests/RHITests/Source/RHITests/Public/RHITextureTests.h
@@ -1424,9 +1424,9 @@ public:
 				};
 
 				const FRHITextureCreateDesc SourceDesc =
-					FRHITextureCreateDesc::Create2D(TestName)
-					.SetExtent(Extent.X, Extent.Y)
-					.SetFormat(PF_A2B10G10R10);
+					FRHITextureCreateDesc::Create2D(TestName, Extent, PF_A2B10G10R10)
+					// Force as non-static
+					.AddFlags(ETextureCreateFlags::RenderTargetable);
 
 				FTextureRHIRef SourceTexture = RHICmdList.CreateTexture(SourceDesc);
 
