Adding FRHIDescriptorRange
- Allows features to allocate a range of descriptors for a specific descriptor type.
- Not implemented on all bindless platforms yet.
- Added unit tests.

#Synced-CL 50610829

diff --git a/Engine/Config/Windows/DataDrivenPlatformInfo.ini b/Engine/Config/Windows/DataDrivenPlatformInfo.ini
--- a/Engine/Config/Windows/DataDrivenPlatformInfo.ini
+++ b/Engine/Config/Windows/DataDrivenPlatformInfo.ini
@@ -121,6 +121,7 @@ bSupportsInstancedStereo=true
 bSupportsPercentageCloserShadows=true
 bSupportsIndexBufferUAVs=true
 bSupportsBindless=true
+bSupportsDescriptorRange=true
 StaticShaderBindingLayoutSupport=RayTracingOnly
 bSupportsROV=true
 bSupportsOIT=true
diff --git a/Engine/Plugins/Tests/RHITests/Shaders/Private/TestDescriptorRange.usf b/Engine/Plugins/Tests/RHITests/Shaders/Private/TestDescriptorRange.usf
new file mode 100644
--- /dev/null
+++ b/Engine/Plugins/Tests/RHITests/Shaders/Private/TestDescriptorRange.usf
@@ -0,0 +1,15 @@
+// Copyright Epic Games, Inc. All Rights Reserved.
+
+#include "/Engine/Private/Common.ush"
+
+FDescriptorRange DescriptorRange;
+
+RWStructuredBuffer<float4> OutputUAV;
+
+[numthreads(THREADGROUP_SIZE_X, THREADGROUP_SIZE_Y, 1)]
+void FillFromTextures(uint3 GroupID : SV_GroupID)
+{
+	uint ResourceIndex = GroupID.x;
+	Texture2D<float4> Tex = GetDescriptorRangeSRV(Texture2D<float4>, DescriptorRange, ResourceIndex);
+	OutputUAV[ResourceIndex] = Tex.Load(uint3(0, 0, 0)).xyzw;
+}
diff --git a/Engine/Plugins/Tests/RHITests/Source/RHITests/Private/RHIBindlessTests.cpp b/Engine/Plugins/Tests/RHITests/Source/RHITests/Private/RHIBindlessTests.cpp
--- a/Engine/Plugins/Tests/RHITests/Source/RHITests/Private/RHIBindlessTests.cpp
+++ b/Engine/Plugins/Tests/RHITests/Source/RHITests/Private/RHIBindlessTests.cpp
@@ -1,8 +1,16 @@
 // Copyright Epic Games, Inc. All Rights Reserved.
 
 #include "RHIBindlessTests.h"
-#include "RHIResourceUtils.h"
+
+#include "CommonRenderResources.h"
+#include "RenderGraphUtils.h"
 #include "RHI.h"
+#include "RHIDescriptorRange.h"
+#include "RHIGPUReadback.h"
+#include "RHIResourceUtils.h"
+#include "RHITestsCommon.h"
+#include "ShaderCompilerCore.h"
+#include "ShaderParameterMacros.h"
 
 bool RHIBindlessTests::Test_ResourceCollection(FRHICommandListImmediate& RHICmdList)
 {
@@ -64,3 +72,167 @@ bool RHIBindlessTests::Test_ResourceCollection(FRHICommandListImmediate& RHICmdL
 
 	return true;
 }
+
+class FDescriptorRangeFillFromTexturesCS : public FGlobalShader
+{
+	DECLARE_GLOBAL_SHADER(FDescriptorRangeFillFromTexturesCS);
+	SHADER_USE_PARAMETER_STRUCT(FDescriptorRangeFillFromTexturesCS, FGlobalShader);
+
+	static constexpr uint32 ThreadGroupSizeX = 1;
+	static constexpr uint32 ThreadGroupSizeY = 1;
+
+	static bool ShouldCompilePermutation(const FShaderPermutationParameters& Parameters)
+	{
+		return IsBindlessEnabledForAnyGraphics(UE::ShaderCompiler::GetBindlessConfiguration(Parameters.Platform))
+			&& FDataDrivenShaderPlatformInfo::GetSupportsDescriptorRange(Parameters.Platform);
+	}
+
+	static void ModifyCompilationEnvironment(const FGlobalShaderPermutationParameters& Parameters, FShaderCompilerEnvironment& OutEnvironment)
+	{
+		FGlobalShader::ModifyCompilationEnvironment(Parameters, OutEnvironment);
+		OutEnvironment.SetDefine(TEXT("THREADGROUP_SIZE_X"), ThreadGroupSizeX);
+		OutEnvironment.SetDefine(TEXT("THREADGROUP_SIZE_Y"), ThreadGroupSizeY);
+
+		OutEnvironment.CompilerFlags.Add(CFLAG_SupportsMinimalBindless);
+	}
+
+	BEGIN_SHADER_PARAMETER_STRUCT(FParameters, )
+		SHADER_PARAMETER_UAV(RWBuffer<float4>, OutputUAV)
+		SHADER_PARAMETER_DESCRIPTOR_RANGE(FDescriptorRange, DescriptorRange)
+	END_SHADER_PARAMETER_STRUCT()
+};
+IMPLEMENT_GLOBAL_SHADER(FDescriptorRangeFillFromTexturesCS, "/Plugin/RHITests/Private/TestDescriptorRange.usf", "FillFromTextures", SF_Compute);
+
+bool RHIBindlessTests::Test_DescriptorRange(FRHICommandListImmediate& RHICmdList)
+{
+	if (!GRHIGlobals.bSupportsBindless || !GRHIGlobals.bSupportsDescriptorRange || !IsBindlessEnabledForAnyGraphics(RHIGetRuntimeBindlessConfiguration(GMaxRHIShaderPlatform)))
+	{
+		return true;
+	}
+
+	const FColor TextureColors[6] =
+	{
+		FColor::White,
+		FColor::Black,
+		FColor::Red,
+		FColor::Green,
+		FColor::Blue,
+		FColor::Yellow,
+	};
+
+	TArray<FTextureRHIRef> Textures;
+	TArray<FShaderResourceViewRHIRef> SRVs;
+
+	{
+		const FRHITextureCreateDesc TextureDesc =
+			FRHITextureCreateDesc::Create2D(TEXT("Test_DescriptorRange"), 1, 1, PF_R8G8B8A8)
+			.SetInitialState(ERHIAccess::SRVCompute)
+			.SetInitActionInitializer();
+
+		for (FColor TextureColor : TextureColors)
+		{
+			uint32 Color = TextureColor.ToPackedABGR();
+
+			FRHITextureInitializer Initializer = RHICmdList.CreateTextureInitializer(TextureDesc);
+			Initializer.GetTexture2DSubresource(0).WriteData(&Color, sizeof(Color));
+			Textures.Emplace(Initializer.Finalize());
+		}
+
+		const FRHIViewDesc::FTextureSRV::FInitializer SRVInit =
+			FRHIViewDesc::CreateTextureSRV()
+			.SetDimension(TextureDesc.Dimension)
+			.SetFormat(PF_R8G8B8A8)
+			.SetMipRange(0, 1);
+
+		for (const FTextureRHIRef& Texture : Textures)
+		{
+			SRVs.Emplace(RHICmdList.CreateShaderResourceView(Texture, SRVInit));
+		}
+	}
+
+	FRHIDescriptorRangeDesc RangeDesc;
+	RangeDesc.Count = Textures.Num();
+	RangeDesc.TypeMask = ERHIDescriptorTypeMask::SRV;
+
+	TArray<FRHIDescriptorResource> Resources;
+	for (const FShaderResourceViewRHIRef& SRV : SRVs)
+	{
+		Resources.Emplace(SRV.GetReference());
+	}
+
+	FDescriptorRangeRHIRef DescriptorRange = RHICmdList.CreateDescriptorRange(RangeDesc, Resources);
+
+	if (!DescriptorRange)
+	{
+		UE_LOG(LogRHIUnitTestCommandlet, Display, TEXT("Unable to allocate RHI Descriptor Range. Assuming it is not implemented for the current platform."));
+		return true;
+	}
+
+	FRHIBufferCreateDesc OutputBufferDesc =
+		FRHIBufferCreateDesc::CreateStructured<FLinearColor>(TEXT("OutputUAV"), Resources.Num())
+		.AddUsage(EBufferUsageFlags::UnorderedAccess)
+		.SetInitialState(ERHIAccess::UAVCompute);
+	FBufferRHIRef OutputBuffer = RHICmdList.CreateBuffer(OutputBufferDesc);
+
+	FUnorderedAccessViewRHIRef OutputUAV = RHICmdList.CreateUnorderedAccessView(
+		OutputBuffer,
+		FRHIViewDesc::CreateBufferUAV()
+			.SetType(FRHIViewDesc::EBufferType::Structured)
+			.SetStride(OutputBufferDesc.Stride).SetNumElements(Resources.Num())
+	);
+
+	{
+		TShaderMapRef< FDescriptorRangeFillFromTexturesCS> ComputeShader(GetGlobalShaderMap(GMaxRHIFeatureLevel));
+		const uint32 NumThreadGroups = FMath::DivideAndRoundUp<uint32>(Resources.Num(), FDescriptorRangeFillFromTexturesCS::ThreadGroupSizeX);
+
+		FDescriptorRangeFillFromTexturesCS::FParameters Parameters;
+		Parameters.OutputUAV = OutputUAV;
+		Parameters.DescriptorRange = DescriptorRange;
+		//Parameters.DescriptorRange_StartIndex = DescriptorRange->GetShaderData().StartIndex;
+		//Parameters.DescriptorRange_Count = DescriptorRange->GetShaderData().Count;
+		
+		FComputeShaderUtils::Dispatch(RHICmdList, ComputeShader, Parameters, FIntVector(NumThreadGroups, 1, 1));
+
+		RHICmdList.Transition(FRHITransitionInfo(OutputUAV, ERHIAccess::UAVCompute, ERHIAccess::CopySrc));
+	}
+
+	FRHIGPUBufferReadback OutputBufferReadback(TEXT("OutputUAVReadback"));
+	OutputBufferReadback.EnqueueCopy(RHICmdList, OutputBuffer, OutputBufferDesc.Size);
+
+	// Sync the GPU. Unfortunately we can't use the fences because not all RHIs implement them yet.
+	RHICmdList.BlockUntilGPUIdle();
+	RHICmdList.FlushResources();
+
+	bool bExpectedResult = true;
+
+	{
+		const FLinearColor* OutputData = reinterpret_cast<const FLinearColor*>(OutputBufferReadback.Lock(OutputBufferDesc.Size));
+		for (int32 Index = 0; Index < Resources.Num(); Index++)
+		{
+			const FLinearColor ExpectedColor = FLinearColor(TextureColors[Index]);
+			const FLinearColor ActualColor = OutputData[Index];
+
+			if (ExpectedColor != ActualColor)
+			{
+				UE_LOG(LogRHIUnitTestCommandlet, Error, TEXT("Test failed. \"Descriptor Range\"\n\nIndex: %d;"
+					" Expected Data: (%.2f, %.2f, %.2f, %.2f); Actual Data: (%.2f, %.2f, %.2f, %.2f);"),
+					Index,
+					ExpectedColor.R, ExpectedColor.G, ExpectedColor.B, ExpectedColor.A,
+					ActualColor.R, ActualColor.G, ActualColor.B, ActualColor.A
+					);
+
+				bExpectedResult = false;
+				break;
+			}
+		}
+		OutputBufferReadback.Unlock();
+
+		if (bExpectedResult)
+		{
+			UE_LOG(LogRHIUnitTestCommandlet, Display, TEXT("Test passed. \"Descriptor Range\""));
+		}
+	}
+
+	return bExpectedResult;
+}
+
diff --git a/Engine/Plugins/Tests/RHITests/Source/RHITests/Private/RHIBindlessTests.h b/Engine/Plugins/Tests/RHITests/Source/RHITests/Private/RHIBindlessTests.h
--- a/Engine/Plugins/Tests/RHITests/Source/RHITests/Private/RHIBindlessTests.h
+++ b/Engine/Plugins/Tests/RHITests/Source/RHITests/Private/RHIBindlessTests.h
@@ -7,4 +7,5 @@ class FRHICommandListImmediate;
 namespace RHIBindlessTests
 {
 	bool Test_ResourceCollection(FRHICommandListImmediate& RHICmdList);
+	bool Test_DescriptorRange(FRHICommandListImmediate& RHICmdList);
 }
diff --git a/Engine/Plugins/Tests/RHITests/Source/RHITests/Private/RHITests.spec.cpp b/Engine/Plugins/Tests/RHITests/Source/RHITests/Private/RHITests.spec.cpp
--- a/Engine/Plugins/Tests/RHITests/Source/RHITests/Private/RHITests.spec.cpp
+++ b/Engine/Plugins/Tests/RHITests/Source/RHITests/Private/RHITests.spec.cpp
@@ -199,7 +199,7 @@ void FAutomationRHITest::Define()
 		It("Vertex shader UAV", [this]()
 		{
 			bool bResult = RunOnRenderThreadSynchronous(FRHIGraphicsUAVTests::Test_GraphicsUAV_VertexShader);
-			TestEqual("Vixel shader UAV failed", bResult, 1);
+			TestEqual("Vertex shader UAV failed", bResult, 1);
 		});
 	});
 
@@ -219,5 +219,10 @@ void FAutomationRHITest::Define()
 			bool bResult = RunOnRenderThreadSynchronous(RHIBindlessTests::Test_ResourceCollection);
 			TestEqual("Resource Collection failed", bResult, 1);
 		});
+		It("RHI Descriptor Range", [this]()
+		{
+			bool bResult = RunOnRenderThreadSynchronous(RHIBindlessTests::Test_DescriptorRange);
+			TestEqual("Descriptor Range failed", bResult, 1);
+		});
 	});
 }
diff --git a/Engine/Plugins/Tests/RHITests/Source/RHITests/Private/RHITestsModule.cpp b/Engine/Plugins/Tests/RHITests/Source/RHITests/Private/RHITestsModule.cpp
--- a/Engine/Plugins/Tests/RHITests/Source/RHITests/Private/RHITestsModule.cpp
+++ b/Engine/Plugins/Tests/RHITests/Source/RHITests/Private/RHITestsModule.cpp
@@ -118,6 +118,7 @@ static bool RunTests_RenderThread(FRHICommandListImmediate& RHICmdList)
 	// ------------------------------------------------
 	{
 		RUN_TEST(RHIBindlessTests::Test_ResourceCollection(RHICmdList));
+		RUN_TEST(RHIBindlessTests::Test_DescriptorRange(RHICmdList));
 	}
 
 	// @todo - add more tests
diff --git a/Engine/Shaders/Public/BindlessResources.ush b/Engine/Shaders/Public/BindlessResources.ush
--- a/Engine/Shaders/Public/BindlessResources.ush
+++ b/Engine/Shaders/Public/BindlessResources.ush
@@ -100,6 +100,21 @@
 	#define UB_RESOURCE_MEMBER_SAMPLER(Type, UBName, Name) Type UBName##_##Name
 #endif
 
+//
+// Descriptor Range support - Experimental feature, this code may completely change in the future.
+
+// .x = StartIndex
+// .y = Count
+typedef uint2 FDescriptorRange;
+
+#if UE_BINDLESS_ENABLED && PLATFORM_SUPPORTS_DESCRIPTOR_RANGE
+	uint GetDescriptorRangeStartIndex(FDescriptorRange InDescriptorRange) { return InDescriptorRange.x; }
+	uint GetDescriptorRangeCount     (FDescriptorRange InDescriptorRange) { return InDescriptorRange.y; }
+
+	#define GetDescriptorRangeSRV(InType, InDescriptorRange, InIndex) GetSRVFromHeap(InType, GetDescriptorRangeStartIndex(InDescriptorRange) + InIndex)
+	#define GetDescriptorRangeUAV(InType, InDescriptorRange, InIndex) GetUAVFromHeap(InType, GetDescriptorRangeStartIndex(InDescriptorRange) + InIndex)
+#endif
+
 //
 // Resource Collection support
 
diff --git a/Engine/Source/Developer/ShaderCompilerCommon/Private/ShaderCompilerCommon.cpp b/Engine/Source/Developer/ShaderCompilerCommon/Private/ShaderCompilerCommon.cpp
--- a/Engine/Source/Developer/ShaderCompilerCommon/Private/ShaderCompilerCommon.cpp
+++ b/Engine/Source/Developer/ShaderCompilerCommon/Private/ShaderCompilerCommon.cpp
@@ -578,6 +578,15 @@ bool UE::ShaderCompilerCommon::RemoveDeadCode(FShaderSource& InOutPreprocessedSh
 	return UE::ShaderCompilerCommon::RemoveDeadCode(InOutPreprocessedShaderSource, RequiredSymbols, OutErrors);
 }
 
+inline bool IsDescriptorRangeParameter(const FShaderParameterParser& ShaderParameterParser, FStringView MemberName)
+{
+	if (const FShaderParameterParser::FParsedShaderParameter* ParsedParameter = ShaderParameterParser.FindParameterInfosUnsafe(MemberName))
+	{
+		return (ParsedParameter->BaseType == UBMT_DESCRIPTOR_RANGE);
+	}
+	return false;
+}
+
 void HandleReflectedGlobalConstantBufferMember(
 	const FShaderCompilerInput& Input,
 	const FShaderParameterParser& ShaderParameterParser,
@@ -589,7 +598,12 @@ void HandleReflectedGlobalConstantBufferMember(
 )
 {
 	FStringView MemberName = InMemberName;
-	const EShaderParameterType ParameterType = FShaderParameterParser::ParseAndRemoveBindlessParameterPrefix(MemberName);
+	EShaderParameterType ParameterType = FShaderParameterParser::ParseAndRemoveBindlessParameterPrefix(MemberName);
+
+	if (ParameterType == EShaderParameterType::LooseData && IsDescriptorRangeParameter(ShaderParameterParser, MemberName))
+	{
+		ParameterType = EShaderParameterType::DescriptorRange;
+	}
 
 	Output.ParameterMap.AddParameterAllocation(
 		MemberName,
@@ -676,7 +690,12 @@ void HandleReflectedRootConstantBufferMember(
 	ShaderParameterParser.ValidateShaderParameterType(Input, InMemberName, ReflectionOffset, ReflectionSize, Output);
 
 	FStringView MemberName = InMemberName;
-	const EShaderParameterType ParameterType = FShaderParameterParser::ParseAndRemoveBindlessParameterPrefix(MemberName);
+	EShaderParameterType ParameterType = FShaderParameterParser::ParseAndRemoveBindlessParameterPrefix(MemberName);
+
+	if (ParameterType == EShaderParameterType::LooseData && IsDescriptorRangeParameter(ShaderParameterParser, MemberName))
+	{
+		ParameterType = EShaderParameterType::DescriptorRange;
+	}
 
 	// Root constant buffer members only need reflection when they are specific to resources
 	if (ParameterType != EShaderParameterType::LooseData)
diff --git a/Engine/Source/Runtime/D3D12RHI/Private/D3D12Adapter.cpp b/Engine/Source/Runtime/D3D12RHI/Private/D3D12Adapter.cpp
--- a/Engine/Source/Runtime/D3D12RHI/Private/D3D12Adapter.cpp
+++ b/Engine/Source/Runtime/D3D12RHI/Private/D3D12Adapter.cpp
@@ -1138,6 +1138,7 @@ void FD3D12Adapter::InitializeDevices()
 			if (GetHighestShaderModel() >= D3D_SHADER_MODEL_6_6 && GetResourceBindingTier() >= D3D12_RESOURCE_BINDING_TIER_3)
 			{
 				GRHIGlobals.bSupportsBindless = (GMaxRHIFeatureLevel > ERHIFeatureLevel::SM5);
+				GRHIGlobals.bSupportsDescriptorRange = GRHIGlobals.bSupportsBindless;
 				UE_LOG(LogD3D12RHI, Log, TEXT("Bindless is supported. MaxNonSamplerDescriptors: %d, MaxSamplerDescriptors: %d"), MaxNonSamplerDescriptors, MaxSamplerDescriptors);
 			}
 
diff --git a/Engine/Source/Runtime/D3D12RHI/Private/D3D12BindlessDescriptors.cpp b/Engine/Source/Runtime/D3D12RHI/Private/D3D12BindlessDescriptors.cpp
--- a/Engine/Source/Runtime/D3D12RHI/Private/D3D12BindlessDescriptors.cpp
+++ b/Engine/Source/Runtime/D3D12RHI/Private/D3D12BindlessDescriptors.cpp
@@ -199,24 +199,21 @@ FRHIDescriptorHandle FD3D12BindlessDescriptorAllocator::AllocateDescriptor(ERHID
 	return Result;
 }
 
-TOptional<FRHIDescriptorAllocation> FD3D12BindlessDescriptorAllocator::AllocateDescriptors(ERHIDescriptorType DescriptorType, uint32 DescriptorCount)
+TOptional<FRHIDescriptorAllocation> FD3D12BindlessDescriptorAllocator::AllocateDescriptors(ERHIDescriptorTypeMask DescriptorTypeMask, uint32 DescriptorCount)
 {
-	if (DescriptorType == ERHIDescriptorType::Sampler)
+	if (SamplerAllocator && SamplerAllocator->HandlesAllocations(DescriptorTypeMask))
 	{
-		if (!AreSamplersBindless())
-		{
-			return TOptional<FRHIDescriptorAllocation>();
-		}
-
+		check(!EnumHasAnyFlags(DescriptorTypeMask, ERHIDescriptorTypeMask::Standard));
 		return SamplerAllocator->Allocate(DescriptorCount);
 	}
 
-	if (!AreResourcesBindless())
+	if (ResourceAllocator && ResourceAllocator->HandlesAllocations(DescriptorTypeMask))
 	{
-		return TOptional<FRHIDescriptorAllocation>();
+		check(!EnumHasAnyFlags(DescriptorTypeMask, ERHIDescriptorTypeMask::Sampler));
+		return ResourceAllocator->Allocate(DescriptorCount);
 	}
 
-	return ResourceAllocator->Allocate(DescriptorCount);
+	return TOptional<FRHIDescriptorAllocation>();
 }
 
 #if D3D12RHI_BINDLESS_RESOURCE_MANAGER_SUPPORTS_RESIZING
@@ -274,6 +271,23 @@ void FD3D12BindlessDescriptorAllocator::FreeDescriptor(FRHIDescriptorHandle Hand
 	}
 }
 
+void FD3D12BindlessDescriptorAllocator::FreeDescriptors(FRHIDescriptorAllocation Allocation, ERHIDescriptorTypeMask TypeMask)
+{
+	if (SamplerAllocator && SamplerAllocator->HandlesAllocations(TypeMask))
+	{
+		SamplerAllocator->Free(Allocation);
+	}
+	else if (ResourceAllocator && ResourceAllocator->HandlesAllocations(TypeMask))
+	{
+		ResourceAllocator->Free(Allocation);
+	}
+	else
+	{
+		// bad configuration?
+		checkNoEntry();
+	}
+}
+
 ////////////////////////////////////////////////////////////////////////////////////////////////////////////////////////
 // FD3D12BindlessResourceManager
 
diff --git a/Engine/Source/Runtime/D3D12RHI/Private/D3D12BindlessDescriptors.h b/Engine/Source/Runtime/D3D12RHI/Private/D3D12BindlessDescriptors.h
--- a/Engine/Source/Runtime/D3D12RHI/Private/D3D12BindlessDescriptors.h
+++ b/Engine/Source/Runtime/D3D12RHI/Private/D3D12BindlessDescriptors.h
@@ -46,10 +46,10 @@ public:
 
 	// Bindless descriptor allocators are stored in the adapter, so descriptor handles can be allocated once and shared for multi-GPU objects
 	FRHIDescriptorHandle AllocateDescriptor(ERHIDescriptorType DescriptorType);
-	TOptional<FRHIDescriptorAllocation> AllocateDescriptors(ERHIDescriptorType DescriptorType, uint32 DescriptorCount);
+	TOptional<FRHIDescriptorAllocation> AllocateDescriptors(ERHIDescriptorTypeMask DescriptorTypeMask, uint32 DescriptorCount);
 
 	void FreeDescriptor(FRHIDescriptorHandle Handle);
-	void FreeDescriptors(ERHIDescriptorType DescriptorType, uint32 Offset);
+	void FreeDescriptors(FRHIDescriptorAllocation Allocation, ERHIDescriptorTypeMask TypeMask);
 
 	FCriticalSection& GetResourceHeapsCS() { return ResourceHeapsCS; }
 
diff --git a/Engine/Source/Runtime/D3D12RHI/Private/D3D12CommandContext.h b/Engine/Source/Runtime/D3D12RHI/Private/D3D12CommandContext.h
--- a/Engine/Source/Runtime/D3D12RHI/Private/D3D12CommandContext.h
+++ b/Engine/Source/Runtime/D3D12RHI/Private/D3D12CommandContext.h
@@ -54,6 +54,7 @@ struct FD3D12DeferredDeleteObject
 		Heap,
 		DescriptorHeap,
 #if PLATFORM_SUPPORTS_BINDLESS_RENDERING
+		DescriptorAllocation,
 		BindlessDescriptor,
 		BindlessDescriptorHeap,
 #endif
@@ -73,11 +74,20 @@ struct FD3D12DeferredDeleteObject
 
 		TUniqueFunction<void()>* Func;
 
+#if PLATFORM_SUPPORTS_BINDLESS_RENDERING
+		struct
+		{
+			FRHIDescriptorAllocation Allocation;
+			ERHIDescriptorTypeMask TypeMask;
+			FD3D12Device* Device;
+		} DescriptorAllocation;
+
 		struct
 		{
 			FRHIDescriptorHandle Handle;
 			FD3D12Device* Device;
 		} BindlessDescriptor;
+#endif
 
 		void* CPUAllocation;
 
@@ -130,6 +140,11 @@ struct FD3D12DeferredDeleteObject
 	{}
 
 #if PLATFORM_SUPPORTS_BINDLESS_RENDERING
+	explicit FD3D12DeferredDeleteObject(FRHIDescriptorAllocation Allocation, ERHIDescriptorTypeMask TypeMask, FD3D12Device* Device)
+		: Type(EType::DescriptorAllocation)
+		, DescriptorAllocation({ Allocation, TypeMask, Device })
+	{}
+
 	explicit FD3D12DeferredDeleteObject(FRHIDescriptorHandle Handle, FD3D12Device* Device)
 		: Type(EType::BindlessDescriptor)
 		, BindlessDescriptor({ Handle, Device })
diff --git a/Engine/Source/Runtime/D3D12RHI/Private/D3D12Commands.cpp b/Engine/Source/Runtime/D3D12RHI/Private/D3D12Commands.cpp
--- a/Engine/Source/Runtime/D3D12RHI/Private/D3D12Commands.cpp
+++ b/Engine/Source/Runtime/D3D12RHI/Private/D3D12Commands.cpp
@@ -7,6 +7,7 @@ D3D12Commands.cpp: D3D RHI commands implementation.
 #include "D3D12RHIPrivate.h"
 #include "ProfilingDebugging/AssetMetadataTrace.h"
 #include "ProfilingDebugging/RealtimeGPUProfiler.h"
+#include "D3D12DescriptorRange.h"
 #include "D3D12ResourceCollection.h"
 #include "D3D12TextureReference.h"
 #include "RHIUniformBufferUtilities.h"
@@ -479,13 +480,19 @@ struct FD3D12ResourceBinder
 	{
 	}
 
+	void SetParameter(uint32 BufferIndex, uint32 BaseIndex, uint32 NumBytes, const uint8* NewValue)
+	{
+		checkSlow(BufferIndex == 0);
+		ConstantBuffer.UpdateConstant(NewValue, BaseIndex, NumBytes);
+	}
+
 #if PLATFORM_SUPPORTS_BINDLESS_RENDERING
 	void SetBindlessHandle(const FRHIDescriptorHandle& Handle, uint32 Offset)
 	{
 		if (Handle.IsValid())
 		{
 			const uint32 BindlessIndex = Handle.GetIndex();
-			ConstantBuffer.UpdateConstant(reinterpret_cast<const uint8*>(&BindlessIndex), Offset, 4);
+			SetParameter(0, Offset, 4, reinterpret_cast<const uint8*>(&BindlessIndex));
 		}
 	}
 #endif
@@ -559,9 +566,9 @@ struct FD3D12ResourceBinder
 		}
 	}
 
+#if PLATFORM_SUPPORTS_BINDLESS_RENDERING
 	void SetResourceCollection(FRHIResourceCollection* ResourceCollection, uint32 Index)
 	{
-#if PLATFORM_SUPPORTS_BINDLESS_RENDERING
 		if (bBindlessResources)
 		{
 			if (FD3D12ResourceCollection* D3D12ResourceCollection = FD3D12CommandContext::RetrieveObject<FD3D12ResourceCollection>(ResourceCollection, GpuIndex))
@@ -584,8 +591,20 @@ struct FD3D12ResourceBinder
 		{
 			checkNoEntry();
 		}
-#endif // PLATFORM_SUPPORTS_BINDLESS_RENDERING
 	}
+
+	void SetDescriptorRange(FRHIDescriptorRange* DescriptorRange, uint32 Offset)
+	{
+		if (bBindlessResources || bBindlessSamplers)
+		{
+			if (FD3D12DescriptorRange* D3D12DescriptorRange = FD3D12CommandContext::RetrieveObject<FD3D12DescriptorRange>(DescriptorRange, GpuIndex))
+			{
+				FRHIDescriptorRangeShaderData ShaderData = D3D12DescriptorRange->GetShaderData();
+				SetParameter(0, Offset, sizeof(ShaderData), reinterpret_cast<const uint8*>(&ShaderData));
+			}
+		}
+	}
+#endif // PLATFORM_SUPPORTS_BINDLESS_RENDERING
 };
 
 static void SetShaderParametersOnContext(
@@ -597,18 +616,13 @@ static void SetShaderParametersOnContext(
 	, TArrayView<const FRHIShaderParameterResource> InResourceParameters
 	, TArrayView<const FRHIShaderParameterResource> InBindlessParameters)
 {
-	FD3D12ConstantBuffer& ConstantBuffer = Context.StageConstantBuffers[ShaderFrequency];
-	FD3D12StateCache& StateCache = Context.StateCache;
-	const uint32 GpuIndex = Context.GetGPUIndex();
+	FD3D12ResourceBinder Binder(Context, ShaderFrequency, GetShaderData(Shader));
 
 	for (const FRHIShaderParameter& Parameter : InParameters)
 	{
-		checkSlow(Parameter.BufferIndex == 0);
-		ConstantBuffer.UpdateConstant(&InParametersData[Parameter.ByteOffset], Parameter.BaseIndex, Parameter.ByteSize);
+		Binder.SetParameter(Parameter.BufferIndex, Parameter.BaseIndex, Parameter.ByteSize, &InParametersData[Parameter.ByteOffset]);
 	}
 
-	FD3D12ResourceBinder Binder(Context, ShaderFrequency, GetShaderData(Shader));
-
 #if PLATFORM_SUPPORTS_BINDLESS_RENDERING
 	for (const FRHIShaderParameterResource& Parameter : InBindlessParameters)
 	{
@@ -618,6 +632,10 @@ static void SetShaderParametersOnContext(
 			{
 				UE::RHI::Private::SetAllUniformBufferResourcesForBindless(Binder, static_cast<FRHIUniformBuffer*>(Resource));
 			}
+			else if (Parameter.Type == FRHIShaderParameterResource::EType::DescriptorRange)
+			{
+				Binder.SetDescriptorRange(static_cast<FRHIDescriptorRange*>(Resource), Parameter.Index);
+			}
 			else
 			{
 				FRHIDescriptorHandle Handle;
@@ -684,11 +702,13 @@ static void SetShaderParametersOnContext(
 			Binder.SetSampler(static_cast<FRHISamplerState*>(Parameter.Resource), Parameter.Index);
 			break;
 		case FRHIShaderParameterResource::EType::UniformBuffer:
-			BindUniformBuffer(Context, Shader, ShaderFrequency, Parameter.Index, FD3D12CommandContext::RetrieveObject<FD3D12UniformBuffer>(Parameter.Resource, GpuIndex));
+			BindUniformBuffer(Context, Shader, ShaderFrequency, Parameter.Index, FD3D12CommandContext::RetrieveObject<FD3D12UniformBuffer>(Parameter.Resource, Binder.GpuIndex));
 			break;
+#if PLATFORM_SUPPORTS_BINDLESS_RENDERING
 		case FRHIShaderParameterResource::EType::ResourceCollection:
 			Binder.SetResourceCollection(static_cast<FRHIResourceCollection*>(Parameter.Resource), Parameter.Index);
 			break;
+#endif
 		default:
 			checkf(false, TEXT("Unhandled resource type?"));
 			break;
diff --git a/Engine/Source/Runtime/D3D12RHI/Private/D3D12DescriptorRange.cpp b/Engine/Source/Runtime/D3D12RHI/Private/D3D12DescriptorRange.cpp
new file mode 100644
--- /dev/null
+++ b/Engine/Source/Runtime/D3D12RHI/Private/D3D12DescriptorRange.cpp
@@ -0,0 +1,159 @@
+// Copyright Epic Games, Inc. All Rights Reserved.
+
+#include "D3D12DescriptorRange.h"
+
+#if PLATFORM_SUPPORTS_BINDLESS_RENDERING
+
+#include "D3D12Adapter.h"
+#include "D3D12RHIPrivate.h"
+
+static FD3D12View* GetViewFromDescriptorResource(const FRHIDescriptorResource& Resource, uint32 GpuIndex)
+{
+	switch (Resource.Type)
+	{
+	case FRHIDescriptorResource::EType::Texture:
+	{
+		FD3D12Texture* D3D12Texture = FD3D12CommandContext::RetrieveTexture(static_cast<FRHITexture*>(Resource.Resource), GpuIndex);
+		return D3D12Texture->GetShaderResourceView();
+	}
+	case FRHIDescriptorResource::EType::ResourceView:
+	{
+		FRHIShaderResourceView* SRV = static_cast<FRHIShaderResourceView*>(Resource.Resource);
+		FD3D12ShaderResourceView_RHI* D3D12ShaderResourceView = FD3D12CommandContext::RetrieveObject<FD3D12ShaderResourceView_RHI>(SRV, GpuIndex);
+		return D3D12ShaderResourceView;
+	}
+	case FRHIDescriptorResource::EType::UnorderedAccessView:
+	{
+		FRHIUnorderedAccessView* UAV = static_cast<FRHIUnorderedAccessView*>(Resource.Resource);
+		FD3D12UnorderedAccessView_RHI* D3D12UnorderedAccessView = FD3D12CommandContext::RetrieveObject<FD3D12UnorderedAccessView_RHI>(UAV, GpuIndex);
+		return D3D12UnorderedAccessView;
+	}
+	}
+
+	return nullptr;
+}
+
+FD3D12DescriptorRange::FD3D12DescriptorRange(FD3D12Device* InParent)
+	: FRHIDescriptorRange(FRHIDescriptorRangeDesc(), {})
+	, FD3D12DeviceChild(InParent)
+	, Allocation(0, 0)
+{
+}
+
+FD3D12DescriptorRange::FD3D12DescriptorRange(FD3D12Device* InParent, const FRHIDescriptorRangeDesc& InDesc, TConstArrayView<FRHIDescriptorResource> InResources, FRHIDescriptorAllocation InAllocation)
+	: FRHIDescriptorRange(InDesc, InResources)
+	, FD3D12DeviceChild(InParent)
+	, Allocation(InAllocation)
+{
+}
+
+FD3D12DescriptorRange::~FD3D12DescriptorRange()
+{
+	if (Allocation.Count > 0)
+	{
+		FD3D12DynamicRHI::GetD3DRHI()->DeferredDelete(Allocation, GetDesc().TypeMask, GetParentDevice());
+	}
+}
+
+FRHIDescriptorRangeShaderData FD3D12DescriptorRange::GetShaderData() const
+{
+	FRHIDescriptorRangeShaderData ShaderData{};
+	ShaderData.StartIndex = Allocation.StartIndex;
+	ShaderData.Count = Allocation.Count;
+	return ShaderData;
+}
+
+void FD3D12DescriptorRange::InitializeDescriptors(TConstArrayView<FRHIDescriptorResource> InResources)
+{
+	const ERHIDescriptorTypeMask DescriptorTypeMask = GetDesc().TypeMask;
+
+	FD3D12Device* Device = GetParentDevice();
+	const uint32 GpuIndex = Device->GetGPUIndex();
+
+	FD3D12BindlessDescriptorManager& Manager = Device->GetBindlessDescriptorManager();
+
+	for (int32 Index = 0; Index < InResources.Num(); Index++)
+	{
+		if (Index < (int32)GetDesc().Count)
+		{
+			const FRHIDescriptorResource& Resource = InResources[Index];
+
+			const FD3D12View* View = GetViewFromDescriptorResource(Resource, GpuIndex);
+			check(EnumHasAllFlags(DescriptorTypeMask, View->GetBindlessHandle().GetTypeMask()));
+
+			const FRHIDescriptorHandle UpdateHandle(View->GetBindlessHandle().GetType(), Allocation.StartIndex + Index);
+			Manager.InitializeDescriptor(UpdateHandle, View);
+		}
+	}
+}
+
+void FD3D12DescriptorRange::UpdateResources(const FRHIContextArray& Contexts, uint32 FirstResourceIndex, TConstArrayView<FRHIDescriptorResource> NewResources)
+{
+	const ERHIDescriptorTypeMask DescriptorTypeMask = GetDesc().TypeMask;
+
+	FD3D12Device* Device = GetParentDevice();
+	const uint32 GpuIndex = Device->GetGPUIndex();
+
+	FD3D12BindlessDescriptorManager& Manager = Device->GetBindlessDescriptorManager();
+
+	for (int32 Index = 0; Index < NewResources.Num(); Index++)
+	{
+		const int32 InternalIndex = FirstResourceIndex + Index;
+		if (InternalIndex < (int32)GetDesc().Count)
+		{
+			const FRHIDescriptorResource& Resource = NewResources[Index];
+			UpdateResource(FirstResourceIndex + Index, Resource);
+
+			const FD3D12View* View = GetViewFromDescriptorResource(Resource, GpuIndex);
+			check(EnumHasAllFlags(DescriptorTypeMask, View->GetBindlessHandle().GetTypeMask()));
+
+			const FRHIDescriptorHandle UpdateHandle(View->GetBindlessHandle().GetType(), Allocation.StartIndex + InternalIndex);
+			Manager.UpdateDescriptor(Contexts, UpdateHandle, View);
+		}
+	}
+}
+
+FDescriptorRangeRHIRef FD3D12DynamicRHI::RHICreateDescriptorRange(FRHICommandListBase& RHICmdList, const FRHIDescriptorRangeDesc& Desc, TConstArrayView<FRHIDescriptorResource> InitialResources)
+{
+	FD3D12DescriptorRange* DescriptorRange = GetAdapter().CreateLinkedObject<FD3D12DescriptorRange>(FRHIGPUMask::All(),
+		[&Desc, &InitialResources](FD3D12Device* Device, FRHIDescriptorRange* FirstLinkedObject)
+		{
+			if (TOptional<FRHIDescriptorAllocation> Allocation = Device->GetBindlessDescriptorAllocator().AllocateDescriptors(Desc.TypeMask, Desc.Count))
+			{
+				return new FD3D12DescriptorRange(Device, Desc, InitialResources, *Allocation);
+			}
+
+			return new FD3D12DescriptorRange(Device);
+		});
+
+	// Because view descriptor initialization can be deferred to the RHI thread, we need to also defer our descriptor copies.
+	TConstArrayView<FRHIDescriptorResource> InitialResourcesCmd = RHICmdList.IsTopOfPipe() ? RHICmdList.AllocArray(InitialResources) : InitialResources;
+
+	RHICmdList.EnqueueLambdaMultiPipe(GetEnabledRHIPipelines(), FRHICommandListBase::EThreadFence::Enabled, TEXT("RHICreateDescriptorRange"),
+		[DescriptorRange, InitialResources = MoveTemp(InitialResourcesCmd)](const FRHIContextArray& Contexts)
+		{
+			for (FD3D12DescriptorRange::FLinkedObjectIterator RangeIt(DescriptorRange); RangeIt; ++RangeIt)
+			{
+				RangeIt->InitializeDescriptors(InitialResources);
+			}
+		});
+
+	return DescriptorRange;
+}
+
+void FD3D12DynamicRHI::RHIUpdateDescriptorRange(FRHICommandListBase& RHICmdList, FRHIDescriptorRange* InDescriptorRange, uint32 FirstResourceIndex, TConstArrayView<FRHIDescriptorResource> NewResources)
+{
+	FD3D12DescriptorRange* DescriptorRange = ResourceCast(InDescriptorRange);
+	TConstArrayView<FRHIDescriptorResource> NewResourcesCmd = RHICmdList.IsTopOfPipe() ? RHICmdList.AllocArray(NewResources) : NewResources;
+
+	RHICmdList.EnqueueLambdaMultiPipe(GetEnabledRHIPipelines(), FRHICommandListBase::EThreadFence::Enabled, TEXT("RHIUpdateDescriptorRange"),
+		[DescriptorRange, FirstResourceIndex, NewResources = MoveTemp(NewResourcesCmd)](const FRHIContextArray& Contexts)
+		{
+			for (FD3D12DescriptorRange::FLinkedObjectIterator RangeIt(DescriptorRange); RangeIt; ++RangeIt)
+			{
+				RangeIt->UpdateResources(Contexts, FirstResourceIndex, NewResources);
+			}
+		});
+}
+
+#endif // PLATFORM_SUPPORTS_BINDLESS_RENDERING
diff --git a/Engine/Source/Runtime/D3D12RHI/Private/D3D12DescriptorRange.h b/Engine/Source/Runtime/D3D12RHI/Private/D3D12DescriptorRange.h
new file mode 100644
--- /dev/null
+++ b/Engine/Source/Runtime/D3D12RHI/Private/D3D12DescriptorRange.h
@@ -0,0 +1,28 @@
+// Copyright Epic Games, Inc. All Rights Reserved.
+
+#pragma once
+
+#include "RHIDescriptorRange.h"
+#include "RHIDescriptorAllocator.h"
+#include "D3D12RHICommon.h"
+
+class FD3D12DescriptorRange : public FRHIDescriptorRange, public FD3D12DeviceChild, public FD3D12LinkedAdapterObject<FD3D12DescriptorRange>
+{
+public:
+	FD3D12DescriptorRange(FD3D12Device* InParent);
+	FD3D12DescriptorRange(FD3D12Device* InParent, const FRHIDescriptorRangeDesc& InDesc, TConstArrayView<FRHIDescriptorResource> InResources, FRHIDescriptorAllocation Allocation);
+	~FD3D12DescriptorRange();
+
+	FRHIDescriptorRangeShaderData GetShaderData() const override;
+
+	void InitializeDescriptors(TConstArrayView<FRHIDescriptorResource> InResources);
+	void UpdateResources(const FRHIContextArray& Contexts, uint32 FirstResourceIndex, TConstArrayView<FRHIDescriptorResource> NewResources);
+
+	FRHIDescriptorAllocation Allocation;
+};
+
+template<>
+struct TD3D12ResourceTraits<FRHIDescriptorRange>
+{
+	using TConcreteType = FD3D12DescriptorRange;
+};
diff --git a/Engine/Source/Runtime/D3D12RHI/Private/D3D12RHI.cpp b/Engine/Source/Runtime/D3D12RHI/Private/D3D12RHI.cpp
--- a/Engine/Source/Runtime/D3D12RHI/Private/D3D12RHI.cpp
+++ b/Engine/Source/Runtime/D3D12RHI/Private/D3D12RHI.cpp
@@ -473,6 +473,12 @@ void FD3D12DynamicRHI::RHIProcessDeleteQueue()
 					ObjectToDelete.D3DObject->Release();
 					break;
 #if PLATFORM_SUPPORTS_BINDLESS_RENDERING
+				case FD3D12DeferredDeleteObject::EType::DescriptorAllocation:
+					ObjectToDelete.DescriptorAllocation.Device->GetBindlessDescriptorAllocator().FreeDescriptors(
+						ObjectToDelete.DescriptorAllocation.Allocation,
+						ObjectToDelete.DescriptorAllocation.TypeMask
+					);
+					break;
 				case FD3D12DeferredDeleteObject::EType::BindlessDescriptor:
 					ObjectToDelete.BindlessDescriptor.Device->GetBindlessDescriptorManager().ImmediateFree(ObjectToDelete.BindlessDescriptor.Handle);
 					break;
diff --git a/Engine/Source/Runtime/D3D12RHI/Private/D3D12RHIPrivate.h b/Engine/Source/Runtime/D3D12RHI/Private/D3D12RHIPrivate.h
--- a/Engine/Source/Runtime/D3D12RHI/Private/D3D12RHIPrivate.h
+++ b/Engine/Source/Runtime/D3D12RHI/Private/D3D12RHIPrivate.h
@@ -449,6 +449,9 @@ public:
 #if PLATFORM_SUPPORTS_BINDLESS_RENDERING
 	virtual FRHIResourceCollectionRef RHICreateResourceCollection(FRHICommandListBase& RHICmdList, TConstArrayView<FRHIResourceCollectionMember> InMembers) override;
 	virtual void RHIUpdateResourceCollection(FRHICommandListBase& RHICmdList, FRHIResourceCollection* InResourceCollection, uint32 InStartIndex, TConstArrayView<FRHIResourceCollectionMember> InMemberUpdates) override;
+
+	virtual FDescriptorRangeRHIRef RHICreateDescriptorRange(FRHICommandListBase& RHICmdList, const FRHIDescriptorRangeDesc& Desc, TConstArrayView<FRHIDescriptorResource> InitialResources) override;
+	virtual void RHIUpdateDescriptorRange(FRHICommandListBase& RHICmdList, FRHIDescriptorRange* Range, uint32 StartIndex, TConstArrayView<FRHIDescriptorResource> Resources) override;
 #endif
 
 	// ID3D12DynamicRHI interface.
diff --git a/Engine/Source/Runtime/Engine/Private/ShaderCompiler/ShaderCompiler.cpp b/Engine/Source/Runtime/Engine/Private/ShaderCompiler/ShaderCompiler.cpp
--- a/Engine/Source/Runtime/Engine/Private/ShaderCompiler/ShaderCompiler.cpp
+++ b/Engine/Source/Runtime/Engine/Private/ShaderCompiler/ShaderCompiler.cpp
@@ -4170,6 +4170,7 @@ void GlobalBeginCompileShader(
 	SET_SHADER_DEFINE(Input.Environment, COMPILER_SUPPORTS_BARYCENTRIC_INTRINSICS, FDataDrivenShaderPlatformInfo::GetSupportsBarycentricsIntrinsics(EShaderPlatform(Target.Platform)));
 	SET_SHADER_DEFINE(Input.Environment, PLATFORM_SUPPORTS_BARYCENTRICS_SEMANTIC, FDataDrivenShaderPlatformInfo::GetSupportsBarycentricsSemantic(EShaderPlatform(Target.Platform)) != ERHIFeatureSupport::Unsupported);
 	SET_SHADER_DEFINE(Input.Environment, PLATFORM_SUPPORTS_BINDLESS, FDataDrivenShaderPlatformInfo::GetSupportsBindless(ShaderPlatform));
+	SET_SHADER_DEFINE(Input.Environment, PLATFORM_SUPPORTS_DESCRIPTOR_RANGE, FDataDrivenShaderPlatformInfo::GetSupportsDescriptorRange(ShaderPlatform));
 	SET_SHADER_DEFINE(Input.Environment, PLATFORM_SUPPORTS_RAYTRACING_CLUSTER_OPS, FDataDrivenShaderPlatformInfo::GetSupportsRayTracingClusterOps(ShaderPlatform));
 	SET_SHADER_DEFINE(Input.Environment, PLATFORM_SUPPORTS_SPECIALIZATION_CONSTANTS, FDataDrivenShaderPlatformInfo::GetSupportsSpecializationConstants(ShaderPlatform));
 
diff --git a/Engine/Source/Runtime/RHI/Private/DataDrivenShaderPlatformInfo.cpp b/Engine/Source/Runtime/RHI/Private/DataDrivenShaderPlatformInfo.cpp
--- a/Engine/Source/Runtime/RHI/Private/DataDrivenShaderPlatformInfo.cpp
+++ b/Engine/Source/Runtime/RHI/Private/DataDrivenShaderPlatformInfo.cpp
@@ -300,6 +300,7 @@ void FGenericDataDrivenShaderPlatformInfo::ParseDataDrivenShaderInfo(const FConf
 	GET_SECTION_BOOL_HELPER(bSupportsFFTBloom);
 	GET_SECTION_BOOL_HELPER(bSupportsVertexShaderLayer);
 	GET_SECTION_BOOL_HELPER(bSupportsBindless);
+	GET_SECTION_BOOL_HELPER(bSupportsDescriptorRange);
 	GET_SECTION_STATIC_SHADER_BINDING_LAYOUT_SUPPORT_HELPER(StaticShaderBindingLayoutSupport);
 	GET_SECTION_BOOL_HELPER(bSupportsVolumeTextureAtomics);
 	GET_SECTION_BOOL_HELPER(bSupportsROV);
@@ -585,6 +586,7 @@ void FGenericDataDrivenShaderPlatformInfo::UpdatePreviewPlatforms()
 				PREVIEW_USE_RUNTIME_VALUE(bInlineRayTracingRequiresBindless);
 				PREVIEW_USE_RUNTIME_VALUE(bSupportsRayTracingShaders);
 				PREVIEW_USE_RUNTIME_VALUE(bSupportsBindless);
+				PREVIEW_USE_RUNTIME_VALUE(bSupportsDescriptorRange);
 				PREVIEW_USE_RUNTIME_VALUE(StaticShaderBindingLayoutSupport);
 
 				// Settings that will never be supported in preview
diff --git a/Engine/Source/Runtime/RHI/Private/RHIDescriptorRange.cpp b/Engine/Source/Runtime/RHI/Private/RHIDescriptorRange.cpp
new file mode 100644
--- /dev/null
+++ b/Engine/Source/Runtime/RHI/Private/RHIDescriptorRange.cpp
@@ -0,0 +1,30 @@
+// Copyright Epic Games, Inc. All Rights Reserved.
+
+#include "RHIDescriptorRange.h"
+
+FRHIDescriptorRange::FRHIDescriptorRange(const FRHIDescriptorRangeDesc& InDesc, TConstArrayView<FRHIDescriptorResource> InResources)
+	: FRHIResource(RRT_DescriptorRange)
+	, Desc(InDesc)
+{
+	Resources.SetNum(InDesc.Count);
+	for (int32 Index = 0; Index < InResources.Num(); Index++)
+	{
+		UpdateResource(Index, InResources[Index]);
+	}
+}
+
+void FRHIDescriptorRange::UpdateResource(int32 Index, FRHIDescriptorResource InResource)
+{
+	if (Index >= 0 && Index < Resources.Num())
+	{
+		Resources[Index] = InResource.Resource;
+	}
+}
+
+void FRHIDescriptorRange::UpdateResources(int32 StartIndex, TConstArrayView<FRHIDescriptorResource> InResources)
+{
+	for (int32 NewResourceIndex = 0; NewResourceIndex < InResources.Num(); NewResourceIndex++)
+	{
+		UpdateResource(StartIndex + NewResourceIndex, InResources[NewResourceIndex]);
+	}
+}
diff --git a/Engine/Source/Runtime/RHI/Private/RHIStrings.cpp b/Engine/Source/Runtime/RHI/Private/RHIStrings.cpp
--- a/Engine/Source/Runtime/RHI/Private/RHIStrings.cpp
+++ b/Engine/Source/Runtime/RHI/Private/RHIStrings.cpp
@@ -1,6 +1,7 @@
 // Copyright Epic Games, Inc. All Rights Reserved.
 
 #include "RHIStrings.h"
+#include "Containers/BitArray.h"
 #include "DataDrivenShaderPlatformInfo.h"
 #include "RHIAccess.h"
 #include "RHIGlobals.h"
@@ -723,6 +724,10 @@ const TCHAR* GetUniformBufferBaseTypeString(EUniformBufferBaseType BaseType)
 		return TEXT("UBMT_REFERENCED_STRUCT");
 	case UBMT_RENDER_TARGET_BINDING_SLOTS:
 		return TEXT("UBMT_RENDER_TARGET_BINDING_SLOTS");
+	case UBMT_RESOURCE_COLLECTION:
+		return TEXT("UBMT_RESOURCE_COLLECTION");
+	case UBMT_DESCRIPTOR_RANGE:
+		return TEXT("UBMT_DESCRIPTOR_RANGE");
 	}
 	return TEXT("");
 }
@@ -753,6 +758,8 @@ const TCHAR* GetShaderCodeResourceBindingTypeName(EShaderCodeResourceBindingType
 	case EShaderCodeResourceBindingType::RWStructuredBuffer:				return TEXT("RWStructuredBuffer");
 	case EShaderCodeResourceBindingType::RWByteAddressBuffer:				return TEXT("RWByteAddressBuffer");
 	case EShaderCodeResourceBindingType::RasterizerOrderedTexture2D:		return TEXT("RasterizerOrderedTexture2D");
+	case EShaderCodeResourceBindingType::ResourceCollection:				return TEXT("ResourceCollection");
+	case EShaderCodeResourceBindingType::DescriptorRange:					return TEXT("DescriptorRange");
 	default:																checkf(false, TEXT("Missing EShaderCodeResourceBindingType %d"), BindingType);
 	}
 	return TEXT("");
@@ -801,3 +808,51 @@ const TCHAR* GetBindlessConfigurationString(ERHIBindlessConfiguration InBindless
 	}
 }
 
+const TCHAR* GetDescriptorTypeString(ERHIDescriptorType InDescriptorType)
+{
+	switch (InDescriptorType)
+	{
+	case ERHIDescriptorType::BufferSRV:             return TEXT("BufferSRV");
+	case ERHIDescriptorType::BufferUAV:             return TEXT("BufferUAV");
+	case ERHIDescriptorType::TypedBufferSRV:        return TEXT("TypedBufferSRV");
+	case ERHIDescriptorType::TypedBufferUAV:        return TEXT("TypedBufferUAV");
+	case ERHIDescriptorType::TextureSRV:            return TEXT("TextureSRV");
+	case ERHIDescriptorType::TextureUAV:            return TEXT("TextureUAV");
+	case ERHIDescriptorType::CBV:                   return TEXT("CBV");
+	case ERHIDescriptorType::Sampler:               return TEXT("Sampler");
+	case ERHIDescriptorType::AccelerationStructure: return TEXT("AccelerationStructure");
+	case ERHIDescriptorType::Invalid:               return TEXT("Invalid");
+	default:
+		checkf(false, TEXT("Unknown DescriptorType: %d"), static_cast<uint32>(InDescriptorType));
+		return TEXT("");
+	}
+}
+
+FString GetDescriptorMaskString(ERHIDescriptorTypeMask InTypeMask)
+{
+	FStringBuilderBase Result;
+
+	if (InTypeMask == ERHIDescriptorTypeMask::None)
+	{
+		Result.Append(TEXT("<empty>"));
+	}
+	else
+	{
+		uint32 IterMask = static_cast<uint32>(InTypeMask);
+		while (IterMask)
+		{
+			const uint32 TypeIndex = FBitSet::GetAndClearNextBit(IterMask);
+			const ERHIDescriptorType DescriptorType = static_cast<ERHIDescriptorType>(TypeIndex);
+
+			if (Result.Len())
+			{
+				Result.Append(TEXT("|"));
+			}
+
+			Result.Append(GetDescriptorTypeString(DescriptorType));
+		}
+	}
+
+	return Result.ToString();
+}
+
diff --git a/Engine/Source/Runtime/RHI/Private/RHIValidation.cpp b/Engine/Source/Runtime/RHI/Private/RHIValidation.cpp
--- a/Engine/Source/Runtime/RHI/Private/RHIValidation.cpp
+++ b/Engine/Source/Runtime/RHI/Private/RHIValidation.cpp
@@ -1153,6 +1153,66 @@ void FValidationRHI::RHIUnlockTexture(FRHICommandListImmediate& RHICmdList, cons
 	return RHI->RHIUnlockTexture(RHICmdList, Arguments);
 }
 
+#if PLATFORM_SUPPORTS_BINDLESS_RENDERING
+inline ERHIDescriptorType GetDescriptorTypeFromDescriptorResource(const FRHIDescriptorResource& InResource)
+{
+	switch (InResource.Type)
+	{
+	case FRHIDescriptorResource::EType::Texture:
+		return static_cast<const FRHITexture*>(InResource.Resource)->GetDefaultBindlessHandle().GetType();
+	case FRHIDescriptorResource::EType::ResourceView:
+		return static_cast<const FRHIShaderResourceView*>(InResource.Resource)->GetBindlessHandle().GetType();
+	case FRHIDescriptorResource::EType::UnorderedAccessView:
+		return static_cast<const FRHIUnorderedAccessView*>(InResource.Resource)->GetBindlessHandle().GetType();
+	}
+
+	return ERHIDescriptorType::Invalid;
+}
+
+inline bool ValidateDescriptorGroups(ERHIDescriptorTypeMask TypeMask)
+{
+	for (ERHIDescriptorTypeMask ValidMaskGroup : GRHIGlobals.DescriptorTypeGroups)
+	{
+		// If we only have a subset of the flags, then we have flags used by different groups, which is invalid.
+		if (!EnumHasAllFlags(ValidMaskGroup, TypeMask) && EnumHasAnyFlags(ValidMaskGroup, TypeMask))
+		{
+			return false;
+		}
+	}
+
+	return true;
+}
+
+inline FString GetAllDescriptorTypeGroupsAsString()
+{
+	FStringBuilderBase Result;
+	for (ERHIDescriptorTypeMask GroupMask : GRHIGlobals.DescriptorTypeGroups)
+	{
+		if (Result.Len())
+		{
+			Result.Append(TEXT(", "));
+		}
+		Result.Append(GetDescriptorMaskString(GroupMask));
+	}
+	return Result.ToString();
+}
+
+void FValidationRHI::ValidateDescriptorTypes(const FRHIDescriptorRangeDesc& RangeDesc, TConstArrayView<FRHIDescriptorResource> Resources)
+{
+	RHI_VALIDATION_CHECK(ValidateDescriptorGroups(RangeDesc.TypeMask),
+		*FString::Printf(TEXT("Invalid descriptor type mask \"%s\" overlaps type groups: %s"), *GetDescriptorMaskString(RangeDesc.TypeMask), *GetAllDescriptorTypeGroupsAsString()));
+
+	for (const FRHIDescriptorResource& Resource : Resources)
+	{
+		const ERHIDescriptorType DescriptorType = GetDescriptorTypeFromDescriptorResource(Resource);
+
+		RHI_VALIDATION_CHECK(DescriptorType != ERHIDescriptorType::Invalid, TEXT("Invalid descriptor type found."));
+		RHI_VALIDATION_CHECK(EnumHasAllFlags(RangeDesc.TypeMask, RHIDescriptorTypeMaskFromType(DescriptorType)),
+			*FString::Printf(TEXT("Unexpected descriptor type found: %s With expected type mask: %s"), GetDescriptorTypeString(DescriptorType), *GetDescriptorMaskString(RangeDesc.TypeMask)));
+	}
+}
+#endif // PLATFORM_SUPPORTS_BINDLESS_RENDERING
+
 class FRHIValidationQueueScope
 {
 	RHIValidation::FOpQueueState* Prev;
diff --git a/Engine/Source/Runtime/RHI/Public/DataDrivenShaderPlatformInfo.h b/Engine/Source/Runtime/RHI/Public/DataDrivenShaderPlatformInfo.h
--- a/Engine/Source/Runtime/RHI/Public/DataDrivenShaderPlatformInfo.h
+++ b/Engine/Source/Runtime/RHI/Public/DataDrivenShaderPlatformInfo.h
@@ -111,6 +111,7 @@ class FGenericDataDrivenShaderPlatformInfo
 	uint32 bSupportsRayTracingShaders : 1;
 	uint32 bSupportsVertexShaderLayer : 1;
 	uint32 bSupportsBindless : 1;
+	uint32 bSupportsDescriptorRange : 1;
 	uint32 StaticShaderBindingLayoutSupport : int32(ERHIStaticShaderBindingLayoutSupport::NumBits);
 	uint32 bSupportsVolumeTextureAtomics : 1;
 	uint32 bSupportsROV : 1;
@@ -757,6 +758,12 @@ public:
 		return Infos[Platform].bSupportsBindless;
 	}
 
+	static inline const bool GetSupportsDescriptorRange(const FStaticShaderPlatform Platform)
+	{
+		check(IsValid(Platform));
+		return Infos[Platform].bSupportsDescriptorRange;
+	}
+
 PRAGMA_DISABLE_DEPRECATION_WARNINGS
 	UE_DEPRECATED(5.7, "GetBindlessSupport is now GetSupportsBindless")
 	static inline const ERHIBindlessSupport GetBindlessSupport(const FStaticShaderPlatform Platform)
diff --git a/Engine/Source/Runtime/RHI/Public/DynamicRHI.h b/Engine/Source/Runtime/RHI/Public/DynamicRHI.h
--- a/Engine/Source/Runtime/RHI/Public/DynamicRHI.h
+++ b/Engine/Source/Runtime/RHI/Public/DynamicRHI.h
@@ -9,6 +9,7 @@ DynamicRHI.h: Dynamically bound Render Hardware Interface definitions.
 #include "CoreTypes.h"
 #include "RHIBufferInitializer.h"
 #include "RHIContext.h"
+#include "RHIDescriptorRange.h"
 #include "RHIFeatureLevel.h"
 #include "RHIShaderLibrary.h"
 #include "RHITextureInitializer.h"
@@ -536,6 +537,15 @@ public:
 	{
 	}
 
+	virtual FDescriptorRangeRHIRef RHICreateDescriptorRange(FRHICommandListBase& RHICmdList, const FRHIDescriptorRangeDesc& Desc, TConstArrayView<FRHIDescriptorResource> InitialResources)
+	{
+		return FDescriptorRangeRHIRef{};
+	}
+
+	virtual void RHIUpdateDescriptorRange(FRHICommandListBase& RHICmdList, FRHIDescriptorRange* DescriptorRange, uint32 FirstResourceIndex, TConstArrayView<FRHIDescriptorResource> NewResources)
+	{
+	}
+
 	/**
 	* Computes the size in memory required by a given texture.
 	*
diff --git a/Engine/Source/Runtime/RHI/Public/RHICommandList.h b/Engine/Source/Runtime/RHI/Public/RHICommandList.h
--- a/Engine/Source/Runtime/RHI/Public/RHICommandList.h
+++ b/Engine/Source/Runtime/RHI/Public/RHICommandList.h
@@ -948,6 +948,16 @@ public:
 		return GDynamicRHI->RHIUpdateResourceCollection(*this, InResourceCollection, InStartIndex, InMemberUpdates);
 	}
 
+	FDescriptorRangeRHIRef CreateDescriptorRange(const FRHIDescriptorRangeDesc& Desc, TConstArrayView<FRHIDescriptorResource> InitialResources)
+	{
+		return GDynamicRHI->RHICreateDescriptorRange(*this, Desc, InitialResources);
+	}
+
+	void UpdateDescriptorRange(FRHIDescriptorRange* DescriptorRange, uint32 FirstResourceIndex, TConstArrayView<FRHIDescriptorResource> NewResources)
+	{
+		GDynamicRHI->RHIUpdateDescriptorRange(*this, DescriptorRange, FirstResourceIndex, NewResources);
+	}
+
 	inline FRayTracingGeometryRHIRef CreateRayTracingGeometry(const FRayTracingGeometryInitializer& Initializer)
 	{
 		return GDynamicRHI->RHICreateRayTracingGeometry(*this, Initializer);
diff --git a/Engine/Source/Runtime/RHI/Public/RHIDefinitions.h b/Engine/Source/Runtime/RHI/Public/RHIDefinitions.h
--- a/Engine/Source/Runtime/RHI/Public/RHIDefinitions.h
+++ b/Engine/Source/Runtime/RHI/Public/RHIDefinitions.h
@@ -564,6 +564,7 @@ enum class EShaderCodeResourceBindingType : uint8
 	RasterizerOrderedTexture2D,
 
 	ResourceCollection,
+	DescriptorRange,
 
 	MAX
 };
@@ -595,6 +596,7 @@ inline bool IsResourceBindingTypeSRV(EShaderCodeResourceBindingType Type)
 	case EShaderCodeResourceBindingType::SamplerState:
 	case EShaderCodeResourceBindingType::TextureMetadata:
 	case EShaderCodeResourceBindingType::RWTextureMetadata:
+	case EShaderCodeResourceBindingType::DescriptorRange:
 		return false;
 	case EShaderCodeResourceBindingType::MAX:
 	case EShaderCodeResourceBindingType::Invalid:
@@ -651,6 +653,7 @@ enum EUniformBufferBaseType : uint8
 	UBMT_RENDER_TARGET_BINDING_SLOTS,
 
 	UBMT_RESOURCE_COLLECTION,
+	UBMT_DESCRIPTOR_RANGE,
 
 	EUniformBufferBaseType_Num,
 	EUniformBufferBaseType_NumBits = 5,
@@ -1053,6 +1056,7 @@ enum ERHIResourceType : uint8
 	RRT_WorkGraphPipelineState,
 	RRT_StreamSourceSlot,
 	RRT_ResourceCollection,
+	RRT_DescriptorRange,
 
 	RRT_Num
 };
@@ -1398,6 +1402,11 @@ struct FRHIDescriptorHandle
 		return Type;
 	}
 
+	ERHIDescriptorTypeMask GetTypeMask() const
+	{
+		return RHIDescriptorTypeMaskFromType(GetType());
+	}
+
 	bool IsValid() const
 	{
 		return Index != MAX_uint32 && Type != ERHIDescriptorType::Invalid;
@@ -1448,6 +1457,14 @@ inline bool IsBindlessFullyEnabled(ERHIBindlessConfiguration Configuration)
 	return Configuration == ERHIBindlessConfiguration::All;
 }
 
+struct FRHIBindlessDescriptorAllocationInfo
+{
+	ERHIDescriptorTypeMask Types{};
+	uint64 Count = 0;
+	uint64 Capacity = 0;
+	uint64 MaxCapacity = 0;
+};
+
 enum class EColorSpaceAndEOTF
 {
 	EUnknown = 0,
@@ -1573,7 +1590,8 @@ inline bool IsShaderParameterTypeReadOnlyRHIResource(EUniformBufferBaseType Base
 		BaseType == UBMT_TEXTURE ||
 		BaseType == UBMT_SRV ||
 		BaseType == UBMT_SAMPLER ||
-		BaseType == UBMT_RESOURCE_COLLECTION;
+		BaseType == UBMT_RESOURCE_COLLECTION ||
+		BaseType == UBMT_DESCRIPTOR_RANGE;
 }
 
 inline bool IsShaderParameterTypeRHIResource(EUniformBufferBaseType BaseType)
diff --git a/Engine/Source/Runtime/RHI/Public/RHIDescriptorRange.h b/Engine/Source/Runtime/RHI/Public/RHIDescriptorRange.h
new file mode 100644
--- /dev/null
+++ b/Engine/Source/Runtime/RHI/Public/RHIDescriptorRange.h
@@ -0,0 +1,127 @@
+// Copyright Epic Games, Inc. All Rights Reserved.
+
+#pragma once
+
+#include "RHIResources.h"
+
+//
+// RHI Descriptor Range - Experimental feature, these types may completely change in the future.
+//
+// Abstraction around allocating a contiguous range of descriptors and loading them in shaders. The descriptor range will manage updating the platform specific
+// bindless descriptor heap with its resources. When bound to a shader, the offset into the bindless descriptor heap will be available in the shader.
+// 
+// Type Mask
+//   Due to limitations on certain platforms, not all resource/view types can be in the same descriptor range. The types that can be used together are exposed
+//   via 'GRHIGlobals.DescriptorTypeGroups'. Each element in the DescriptorTypeGroups array corresponds to a set of descriptor types (via mask) that can be
+//   used together. Given how nuanced some of the descriptor types are, this may create friction in some uses of descriptor ranges.
+// 
+//   To make type mask support easier, consider using different descriptor ranges for read-only vs read-write resources.
+// 
+// Resource Lifetimes
+//   The Descriptor Range will keep the resources assigned to it alive until the descriptor range is destroyed.
+// 
+// Support checks for this feature:
+//   Shader support: FDataDrivenShaderPlatformInfo::GetSupportsDescriptorRange
+//   Shader define:  PLATFORM_SUPPORTS_DESCRIPTOR_RANGE
+//   Runtime:        GRHIGlobals.bSupportsDescriptorRange
+// (These must be checked in addition to the normal bindless checks)
+// 
+// Creation:
+//   FRHIDescriptorRangeDesc Desc{ Count, TypeMask };
+//   TArray<FRHIDescriptorResource> Resources = ...;
+//   FDescriptorRangeRHIRef DescriptorRange = RHICmdList.CreateDescriptorRange(Desc, Resources);
+// 
+// Updating:
+//   TArray<FRHIDescriptorResource> Resources = ...;
+//   RHICmdList.UpdateDescriptorRange(DescriptorRange, Offset, Resources);
+// 
+// Shader binding:
+//   SHADER_PARAMETER_DESCRIPTOR_RANGE(FDescriptorRange, DescriptorRange)
+// 
+// Shader access:
+//   Texture2D MyTexture = GetDescriptorRangeSRV(Texture2D, DescriptorRange, Offset);
+// 
+
+struct FRHIDescriptorRangeDesc
+{
+	// Number of descriptors in the range.
+	uint32 Count = 0;
+
+	// Mask of different descriptor types in the range. Types cannot cross groups exposed in GRHIGlobals.DescriptorTypeGroups
+	ERHIDescriptorTypeMask TypeMask = ERHIDescriptorTypeMask::None;
+};
+
+// The data that is bound in the shaders.
+struct FRHIDescriptorRangeShaderData
+{
+	uint32 StartIndex;
+	uint32 Count;
+};
+static_assert(sizeof(FRHIDescriptorRangeShaderData) == sizeof(void*), "FRHIDescriptorRangeShaderData must be exactly the same size as a pointer");
+
+// Container for descriptor sources. Sources are Textures, SRVs and UAVs.
+// Samplers may be added in the future.
+struct FRHIDescriptorResource
+{
+	enum class EType : uint8
+	{
+		None,
+		Texture,
+		ResourceView,
+		UnorderedAccessView,
+	};
+
+	FRHIDescriptorResource() = default;
+	FRHIDescriptorResource(EType InType, FRHIResource* InResource)
+		: Resource(InResource)
+		, Type(InType)
+	{
+	}
+	FRHIDescriptorResource(FRHITexture* InTexture)
+		: FRHIDescriptorResource(FRHIDescriptorResource::EType::Texture, InTexture)
+	{
+	}
+	FRHIDescriptorResource(FRHIShaderResourceView* InView)
+		: FRHIDescriptorResource(FRHIDescriptorResource::EType::ResourceView, InView)
+	{
+	}
+	FRHIDescriptorResource(FRHIUnorderedAccessView* InUAV)
+		: FRHIDescriptorResource(FRHIDescriptorResource::EType::UnorderedAccessView, InUAV)
+	{
+	}
+
+	bool operator==(const FRHIDescriptorResource& Other) const
+	{
+		return Resource == Other.Resource
+			&& Type == Other.Type;
+	}
+
+	FRHIResource* Resource = nullptr;
+	EType         Type = EType::None;
+};
+
+class FRHIDescriptorRange : public FRHIResource
+{
+public:
+	const FRHIDescriptorRangeDesc& GetDesc() const
+	{
+		return Desc;
+	}
+
+	bool IsValid() const
+	{
+		return GetDesc().TypeMask != ERHIDescriptorTypeMask::None;
+	}
+
+	virtual FRHIDescriptorRangeShaderData GetShaderData() const = 0;
+
+protected:
+	RHI_API FRHIDescriptorRange(const FRHIDescriptorRangeDesc& InDesc, TConstArrayView<FRHIDescriptorResource> InResources);
+
+	RHI_API void UpdateResource(int32 Index, FRHIDescriptorResource InResource);
+	RHI_API void UpdateResources(int32 StartIndex, TConstArrayView<FRHIDescriptorResource> InResources);
+
+	const FRHIDescriptorRangeDesc Desc;
+
+	TArray<TRefCountPtr<FRHIResource>> Resources;
+};
diff --git a/Engine/Source/Runtime/RHI/Public/RHIFwd.h b/Engine/Source/Runtime/RHI/Public/RHIFwd.h
--- a/Engine/Source/Runtime/RHI/Public/RHIFwd.h
+++ b/Engine/Source/Runtime/RHI/Public/RHIFwd.h
@@ -58,6 +58,7 @@ class FRHIComputePipelineState;
 class FRHIComputeShader;
 class FRHICustomPresent;
 class FRHIDepthStencilState;
+class FRHIDescriptorRange;
 class FRHIGeometryShader;
 class FRHIGPUFence;
 class FRHIGraphicsPipelineState;
@@ -104,6 +105,7 @@ using FComputePipelineStateRHIRef      = TRefCountPtr<FRHIComputePipelineState>;
 using FComputeShaderRHIRef             = TRefCountPtr<FRHIComputeShader>;
 using FCustomPresentRHIRef             = TRefCountPtr<FRHICustomPresent>;
 using FDepthStencilStateRHIRef         = TRefCountPtr<FRHIDepthStencilState>;
+using FDescriptorRangeRHIRef           = TRefCountPtr<FRHIDescriptorRange>;
 using FGeometryShaderRHIRef            = TRefCountPtr<FRHIGeometryShader>;
 using FGPUFenceRHIRef                  = TRefCountPtr<FRHIGPUFence>;
 using FGraphicsPipelineStateRHIRef     = TRefCountPtr<FRHIGraphicsPipelineState>;
diff --git a/Engine/Source/Runtime/RHI/Public/RHIGlobals.h b/Engine/Source/Runtime/RHI/Public/RHIGlobals.h
--- a/Engine/Source/Runtime/RHI/Public/RHIGlobals.h
+++ b/Engine/Source/Runtime/RHI/Public/RHIGlobals.h
@@ -682,9 +682,12 @@ PRAGMA_DISABLE_DEPRECATION_WARNINGS
 	ERHIBindlessSupport BindlessSupport = ERHIBindlessSupport::Unsupported;
 PRAGMA_ENABLE_DEPRECATION_WARNINGS
 
-	/** Whether dynamic (bindless) resources are supported */
+	/** Whether dynamic (bindless) resources are supported. Can be set even if bindless is not enabled. */
 	bool bSupportsBindless = false;
 
+	/** Whether bindless descriptor ranges are supported. Can be set even if bindless is not enabled. */
+	bool bSupportsDescriptorRange = false;
+
 	/** Array of descriptor type masks showing which descriptor types can be allocated together. */
 	TArray<ERHIDescriptorTypeMask> DescriptorTypeGroups;
 
diff --git a/Engine/Source/Runtime/RHI/Public/RHIShaderParameters.h b/Engine/Source/Runtime/RHI/Public/RHIShaderParameters.h
--- a/Engine/Source/Runtime/RHI/Public/RHIShaderParameters.h
+++ b/Engine/Source/Runtime/RHI/Public/RHIShaderParameters.h
@@ -5,6 +5,7 @@
 #include "Containers/Array.h"
 #include "RHIResources.h"
 #include "RHIResourceCollection.h"
+#include "RHIDescriptorRange.h"
 
 #define RHI_VALIDATE_BATCHED_SHADER_PARAMETERS DO_CHECK
 
@@ -38,6 +39,7 @@ struct FRHIShaderParameterResource
 		Sampler,
 		UniformBuffer,
 		ResourceCollection,
+		DescriptorRange,
 	};
 
 	FRHIShaderParameterResource() = default;
@@ -71,6 +73,10 @@ struct FRHIShaderParameterResource
 		: FRHIShaderParameterResource(FRHIShaderParameterResource::EType::ResourceCollection, InResourceCollection, InIndex)
 	{
 	}
+	FRHIShaderParameterResource(FRHIDescriptorRange* InDescriptorRange, uint16 InIndex)
+		: FRHIShaderParameterResource(FRHIShaderParameterResource::EType::DescriptorRange, InDescriptorRange, InIndex)
+	{
+	}
 
 	bool operator == (const FRHIShaderParameterResource& Other) const
 	{
@@ -327,6 +333,11 @@ struct FRHIBatchedShaderParameters
 		AddResourceParameter(ResourceCollection, (uint16)Index);
 	}
 
+	inline void SetDescriptorRange(uint32 Index, FRHIDescriptorRange* DescriptorRange)
+	{
+		AddResourceParameter(DescriptorRange, (uint16)Index);
+	}
+
 	inline void SetBindlessTexture(uint32 Index, FRHITexture* Texture)
 	{
 		AddBindlessParameter(Texture, (uint16)Index);
diff --git a/Engine/Source/Runtime/RHI/Public/RHIStrings.h b/Engine/Source/Runtime/RHI/Public/RHIStrings.h
--- a/Engine/Source/Runtime/RHI/Public/RHIStrings.h
+++ b/Engine/Source/Runtime/RHI/Public/RHIStrings.h
@@ -37,6 +37,9 @@ RHI_API bool GetBindlessConfigurationFromString(FStringView InString, ERHIBindle
 /** Get the string for the given BindlessConfiguration or return an empty string if invalid. */
 RHI_API const TCHAR* GetBindlessConfigurationString(ERHIBindlessConfiguration InBindlessConfiguration);
 
+RHI_API const TCHAR* GetDescriptorTypeString(ERHIDescriptorType InDescriptorType);
+RHI_API FString GetDescriptorMaskString(ERHIDescriptorTypeMask InTypeMask);
+
 /** Stringifies ERHIFeatureLevel */
 RHI_API FString LexToString(ERHIFeatureLevel::Type Level);
 
diff --git a/Engine/Source/Runtime/RHI/Public/RHIValidation.h b/Engine/Source/Runtime/RHI/Public/RHIValidation.h
--- a/Engine/Source/Runtime/RHI/Public/RHIValidation.h
+++ b/Engine/Source/Runtime/RHI/Public/RHIValidation.h
@@ -7,6 +7,7 @@
 #pragma once 
 
 #include "RHI.h"
+#include "RHIDescriptorRange.h"
 #include "RHIValidationCommon.h"
 #include "RHIValidationUtils.h"
 #include "DataDrivenShaderPlatformInfo.h"
@@ -1045,6 +1046,21 @@ public:
 		return RHI->RHIWriteResourceGPUAddresses(RHICmdList, GPUIndex, Resources, AddressesPtr);
 	}
 
+#if PLATFORM_SUPPORTS_BINDLESS_RENDERING
+	RHI_API static void ValidateDescriptorTypes(const FRHIDescriptorRangeDesc& RangeDesc, TConstArrayView<FRHIDescriptorResource> Resources);
+
+	virtual FDescriptorRangeRHIRef RHICreateDescriptorRange(FRHICommandListBase& RHICmdList, const FRHIDescriptorRangeDesc& Desc, TConstArrayView<FRHIDescriptorResource> InitialResources) override final
+	{
+		ValidateDescriptorTypes(Desc, InitialResources);
+		return RHI->RHICreateDescriptorRange(RHICmdList, Desc, InitialResources);
+	}
+	virtual void RHIUpdateDescriptorRange(FRHICommandListBase& RHICmdList, FRHIDescriptorRange* Range, uint32 StartIndex, TConstArrayView<FRHIDescriptorResource> Resources) override final
+	{
+		ValidateDescriptorTypes(Range->GetDesc(), Resources);
+		RHI->RHIUpdateDescriptorRange(RHICmdList, Range, StartIndex, Resources);
+	}
+#endif
+
 //protected:
 	static void ReportValidationFailure(const TCHAR* InMessage, RHIValidation::FTracker* Tracker);
 	static void ReportValidationFailure(const TCHAR* InMessage, FValidationComputeContext* Context);
diff --git a/Engine/Source/Runtime/RHICore/Internal/RHIUniformBufferDataShared.h b/Engine/Source/Runtime/RHICore/Internal/RHIUniformBufferDataShared.h
--- a/Engine/Source/Runtime/RHICore/Internal/RHIUniformBufferDataShared.h
+++ b/Engine/Source/Runtime/RHICore/Internal/RHIUniformBufferDataShared.h
@@ -97,6 +97,7 @@ namespace UE::RHICore
 		// First copy wholesale
 		FMemory::Memcpy(DestinationData, SourceData, Layout.ConstantBufferSize);
 
+#if PLATFORM_SUPPORTS_BINDLESS_RENDERING
 		if (bAllowBindless)
 		{
 			FUniformDataReader Reader(SourceData);
@@ -104,13 +105,25 @@ namespace UE::RHICore
 			// Then copy indices over
 			for (const FRHIUniformBufferResource& Resource : Layout.Resources)
 			{
-				const FRHIDescriptorHandle Handle = GetBindlessResourceHandle(Reader, Resource);
-				if (Handle.IsValid())
+				if (Resource.MemberType == UBMT_DESCRIPTOR_RANGE)
 				{
-					const uint32 BindlessIndex = Handle.GetIndex();
-					FMemory::Memcpy(reinterpret_cast<uint8*>(DestinationData) + Resource.MemberOffset, &BindlessIndex, sizeof(BindlessIndex));
+					if (FRHIDescriptorRange* DescriptorRange = Reader.Read<FRHIDescriptorRange*>(Resource))
+					{
+						FRHIDescriptorRangeShaderData ShaderData = DescriptorRange->GetShaderData();
+						FMemory::Memcpy(reinterpret_cast<uint8*>(DestinationData) + Resource.MemberOffset, &ShaderData, sizeof(ShaderData));
+					}
+				}
+				else
+				{
+					const FRHIDescriptorHandle Handle = GetBindlessResourceHandle(Reader, Resource);
+					if (Handle.IsValid())
+					{
+						const uint32 BindlessIndex = Handle.GetIndex();
+						FMemory::Memcpy(reinterpret_cast<uint8*>(DestinationData) + Resource.MemberOffset, &BindlessIndex, sizeof(BindlessIndex));
+					}
 				}
 			}
 		}
+#endif
 	}
 }
diff --git a/Engine/Source/Runtime/RHICore/Private/RHIDescriptorAllocator.cpp b/Engine/Source/Runtime/RHICore/Private/RHIDescriptorAllocator.cpp
--- a/Engine/Source/Runtime/RHICore/Private/RHIDescriptorAllocator.cpp
+++ b/Engine/Source/Runtime/RHICore/Private/RHIDescriptorAllocator.cpp
@@ -361,3 +361,20 @@ void FRHIOffsetHeapDescriptorAllocator::Free(const FRHIDescriptorHandle InHandle
 		FRHIHeapDescriptorAllocator::Free(AdjustedHandle);
 	}
 }
+
+TOptional<FRHIDescriptorAllocation> FRHIOffsetHeapDescriptorAllocator::Allocate(uint32 NumDescriptors)
+{
+	TOptional<FRHIDescriptorAllocation> Allocation = FRHIHeapDescriptorAllocator::Allocate(NumDescriptors);
+	if (Allocation)
+	{
+		Allocation->StartIndex += HeapOffset;
+	}
+	return Allocation;
+}
+
+void FRHIOffsetHeapDescriptorAllocator::Free(FRHIDescriptorAllocation Allocation)
+{
+	Allocation.StartIndex -= HeapOffset;
+	FRHIHeapDescriptorAllocator::Free(Allocation);
+}
+
diff --git a/Engine/Source/Runtime/RHICore/Public/RHIDescriptorAllocator.h b/Engine/Source/Runtime/RHICore/Public/RHIDescriptorAllocator.h
--- a/Engine/Source/Runtime/RHICore/Public/RHIDescriptorAllocator.h
+++ b/Engine/Source/Runtime/RHICore/Public/RHIDescriptorAllocator.h
@@ -1,4 +1,5 @@
 // Copyright Epic Games, Inc. All Rights Reserved.
+
 #pragma once
 
 #include "Containers/Array.h"
@@ -130,9 +131,13 @@ public:
 	RHICORE_API FRHIDescriptorHandle Allocate(ERHIDescriptorType InType);
 	RHICORE_API void Free(FRHIDescriptorHandle InHandle);
 
+	RHICORE_API TOptional<FRHIDescriptorAllocation> Allocate(uint32 NumDescriptors);
+	RHICORE_API void Free(FRHIDescriptorAllocation Allocation);
+
 	using FRHIHeapDescriptorAllocator::GetCapacity;
 	using FRHIHeapDescriptorAllocator::GetTypeMask;
 	using FRHIHeapDescriptorAllocator::HandlesAllocation;
+	using FRHIHeapDescriptorAllocator::HandlesAllocations;
 
 private:
 	// Offset from start of heap we belong to
diff --git a/Engine/Source/Runtime/RenderCore/Private/ShaderParameterMetadata.cpp b/Engine/Source/Runtime/RenderCore/Private/ShaderParameterMetadata.cpp
--- a/Engine/Source/Runtime/RenderCore/Private/ShaderParameterMetadata.cpp
+++ b/Engine/Source/Runtime/RenderCore/Private/ShaderParameterMetadata.cpp
@@ -109,6 +109,9 @@ void FShaderParametersMetadata::FMember::GenerateShaderParameterType(
 			Result = TEXT("fixed");
 		}
 		break;
+	case UBMT_DESCRIPTOR_RANGE:
+		Result = TEXT("FDescriptorRange");
+		break;
 	default:
 		UE_LOG(LogShaders, Fatal, TEXT("Unrecognized uniform buffer struct member base type."));
 	};
@@ -297,6 +300,7 @@ const TCHAR* const kShaderParameterMacroNames[] = {
 	TEXT("RENDER_TARGET_BINDING_SLOTS"), // UBMT_RENDER_TARGET_BINDING_SLOTS,
 
 	TEXT("RESOURCE_COLLECTION"), // UBMT_RESOURCE_COLLECTION,
+	TEXT("DESCRIPTOR_RANGE"), // UBMT_DESCRIPTOR_RANGE,
 };
 
 static_assert(UE_ARRAY_COUNT(kShaderParameterMacroNames) == int32(EUniformBufferBaseType_Num), "Shader parameter enum does not match name macro name array.");
@@ -387,6 +391,10 @@ EShaderCodeResourceBindingType ParseShaderResourceBindingType(const TCHAR* Shade
 	{
 		BindingType = EShaderCodeResourceBindingType::Buffer;
 	}
+	else if (!bIsRWResource && ShaderTypeLength == 16 && FCString::Strncmp(ComparedShaderType, TEXT("FDescriptorRange"), ShaderTypeLength) == 0)
+	{
+		BindingType = EShaderCodeResourceBindingType::DescriptorRange;
+	}
 	else if (bIsRasterizerOrderedResource)
 	{
 		BindingType = EShaderCodeResourceBindingType::RasterizerOrderedTexture2D;
@@ -788,6 +796,10 @@ void FShaderParametersMetadata::InitializeLayout(FRHIUniformBufferLayoutInitiali
 				{
 					bIsValidBindingType = BaseType == UBMT_RESOURCE_COLLECTION;
 				}
+				else if (BindingType == EShaderCodeResourceBindingType::DescriptorRange)
+				{
+					bIsValidBindingType = BaseType == UBMT_DESCRIPTOR_RANGE;
+				}
 				else
 				{
 					//unimplemented();
diff --git a/Engine/Source/Runtime/RenderCore/Private/ShaderParameterParser.cpp b/Engine/Source/Runtime/RenderCore/Private/ShaderParameterParser.cpp
--- a/Engine/Source/Runtime/RenderCore/Private/ShaderParameterParser.cpp
+++ b/Engine/Source/Runtime/RenderCore/Private/ShaderParameterParser.cpp
@@ -299,7 +299,18 @@ bool FShaderParameterParser::ParseParameters(
 				ParsedParameter.PrecisionModifier = Member.GetPrecision();
 				ParsedParameter.NumRows = Member.GetNumRows();
 				ParsedParameter.NumColumns = Member.GetNumColumns();
-				ParsedParameter.MemberSize = Member.IsVariableNativeType() ? Member.GetMemberSize() : 0u;
+				if (Member.IsVariableNativeType())
+				{
+					ParsedParameter.MemberSize = Member.GetMemberSize();
+				}
+				else if (Member.GetBaseType() == UBMT_DESCRIPTOR_RANGE)
+				{
+					ParsedParameter.MemberSize = sizeof(uint32) * 2;
+				}
+				else
+				{
+					ParsedParameter.MemberSize = 0;
+				}
 				ParsedParameters.Add(ShaderBindingName, ParsedParameter);
 			});
 	}
@@ -462,6 +473,7 @@ bool FShaderParameterParser::ParseParameters(
 								ParsedParameter.BaseType == UBMT_INT32 ||
 								ParsedParameter.BaseType == UBMT_UINT32 ||
 								ParsedParameter.BaseType == UBMT_FLOAT32 ||
+								ParsedParameter.BaseType == UBMT_DESCRIPTOR_RANGE ||
 								ParsedBindlessType != EShaderParameterType::LooseData ||
 								BindlessConversionType != EBindlessConversionType::None;
 
@@ -1140,6 +1152,7 @@ bool FShaderParameterParser::MoveShaderParametersToRootConstantBuffer(
 	if (RootParametersStructure && bMoveGlobalsToRootCB)
 	{
 		FStringBuilderBase ConstantBufferMembers;
+		bool bNeedDescriptorRangeType = false;
 
 		IterateShaderParameterMembers(
 			*RootParametersStructure,
@@ -1154,6 +1167,12 @@ bool FShaderParameterParser::MoveShaderParametersToRootConstantBuffer(
 				const uint32 ConstantRegister = ByteOffset / 16;
 				const TCHAR* ConstantSwizzle = GetConstantSwizzle(ByteOffset);
 
+				// Detect if we need to add the typedef before the constant buffer.
+				if (ParsedParameter->BaseType == UBMT_DESCRIPTOR_RANGE)
+				{
+					bNeedDescriptorRangeType = true;
+				}
+
 #define SVARG(N) N.Len(), N.GetData()
 				if (IsParameterBindless(ParsedParameter->RootCBParameterType))
 				{
@@ -1199,6 +1218,12 @@ bool FShaderParameterParser::MoveShaderParametersToRootConstantBuffer(
 		{
 			FStringBuilderBase ConstantBufferCode;
 
+			// Hack-ish approach to make sure the descriptor range type is defined before the constant buffer.
+			if (bNeedDescriptorRangeType)
+			{
+				ConstantBufferCode.Append(TEXT("typedef uint2 FDescriptorRange;\n"));
+			}
+
 			ConstantBufferCode << PlatformConfiguration.StableConstantBufferType << TEXT(" ") << FShaderParametersMetadata::kRootUniformBufferBindingName << PlatformConfiguration.GetStableConstantBufferRegisterString(ShaderFrequency) << TEXT("\n{\n");
 			ConstantBufferCode << ConstantBufferMembers;
 			ConstantBufferCode << TEXT("}\n\n");
diff --git a/Engine/Source/Runtime/RenderCore/Private/ShaderParameterStruct.cpp b/Engine/Source/Runtime/RenderCore/Private/ShaderParameterStruct.cpp
--- a/Engine/Source/Runtime/RenderCore/Private/ShaderParameterStruct.cpp
+++ b/Engine/Source/Runtime/RenderCore/Private/ShaderParameterStruct.cpp
@@ -86,7 +86,8 @@ struct FShaderParameterStructBindingContext
 				BaseType == UBMT_TEXTURE ||
 				BaseType == UBMT_SRV ||
 				BaseType == UBMT_UAV ||
-				BaseType == UBMT_SAMPLER);
+				BaseType == UBMT_SAMPLER ||
+				BaseType == UBMT_DESCRIPTOR_RANGE);
 
 			const bool bIsRDGResource =
 				IsRDGResourceReferenceShaderParameterType(BaseType);
@@ -208,7 +209,7 @@ struct FShaderParameterStructBindingContext
 				}
 				else if (bIsRHIResource || bIsRDGResource)
 				{
-					if (IsParameterBindless(ParameterAllocation->Type))
+					if (IsParameterBindless(ParameterAllocation->Type) || ParameterAllocation->Type == EShaderParameterType::DescriptorRange)
 					{
 						FShaderParameterBindings::FBindlessResourceParameter Parameter;
 						Parameter.GlobalConstantOffset = ParameterAllocation->BaseIndex;
@@ -674,6 +675,12 @@ FRHIShaderParameterResource ExtractShaderParameterResource(FShaderParameterReade
 		checkSlow(SamplerState);
 		return FRHIShaderParameterResource(SamplerState, GetParameterIndex(Parameter));
 	}
+	case UBMT_DESCRIPTOR_RANGE:
+	{
+		FRHIDescriptorRange* DescriptorRange = Reader.Read<FRHIDescriptorRange*>(Parameter);
+		checkSlow(DescriptorRange);
+		return FRHIShaderParameterResource(DescriptorRange, GetParameterIndex(Parameter));
+	}
 	case UBMT_RDG_TEXTURE:
 	{
 		FRDGTexture* RDGTexture = Reader.Read<FRDGTexture*>(Parameter);
diff --git a/Engine/Source/Runtime/RenderCore/Public/ShaderCore.h b/Engine/Source/Runtime/RenderCore/Public/ShaderCore.h
--- a/Engine/Source/Runtime/RenderCore/Public/ShaderCore.h
+++ b/Engine/Source/Runtime/RenderCore/Public/ShaderCore.h
@@ -256,6 +256,8 @@ enum class EShaderParameterType : uint8
 	BindlessSRV,
 	BindlessUAV,
 
+	DescriptorRange,
+
 	Num
 };
 
@@ -269,6 +271,7 @@ enum class EShaderParameterTypeMask : uint16
 	BindlessSamplerMask = 1 << uint16(EShaderParameterType::BindlessSampler),
 	BindlessSRVMask = 1 << uint16(EShaderParameterType::BindlessSRV),
 	BindlessUAVMask = 1 << uint16(EShaderParameterType::BindlessUAV),
+	DescriptorRangeMask = 1 << uint16(EShaderParameterType::DescriptorRange),
 };
 ENUM_CLASS_FLAGS(EShaderParameterTypeMask);
 
diff --git a/Engine/Source/Runtime/RenderCore/Public/ShaderParameterMacros.h b/Engine/Source/Runtime/RenderCore/Public/ShaderParameterMacros.h
--- a/Engine/Source/Runtime/RenderCore/Public/ShaderParameterMacros.h
+++ b/Engine/Source/Runtime/RenderCore/Public/ShaderParameterMacros.h
@@ -1672,6 +1672,14 @@ private:
 #define SHADER_PARAMETER_SAMPLER_ARRAY(ShaderType,MemberName, ArrayDecl) \
 	INTERNAL_SHADER_PARAMETER_EXPLICIT(UBMT_SAMPLER, TShaderResourceParameterTypeInfo<FRHISamplerState* ArrayDecl>, FRHISamplerState*,MemberName,ArrayDecl,,EShaderPrecisionModifier::Float,TEXT(#ShaderType),false)
 
+/** Adds a descriptor range.
+ *
+ * Example:
+ *	SHADER_PARAMETER_DESCRIPTOR_RANGE(FDescriptorRange, MyDescriptorRange)
+ */
+#define SHADER_PARAMETER_DESCRIPTOR_RANGE(ShaderType,MemberName) \
+	INTERNAL_SHADER_PARAMETER_EXPLICIT(UBMT_DESCRIPTOR_RANGE,TShaderResourceParameterTypeInfo<FRHIDescriptorRange*>,FRHIDescriptorRange*,MemberName,,=nullptr,EShaderPrecisionModifier::Float,TEXT(#ShaderType),false)
+
 /** Adds a render graph tracked texture.
  *
  * Example:
