Fix failure in RHICopyTexture test

#Synced-CL 51293957

diff --git a/Engine/Plugins/Tests/RHITests/Source/RHITests/Public/RHITextureTests.h b/Engine/Plugins/Tests/RHITests/Source/RHITests/Public/RHITextureTests.h
--- a/Engine/Plugins/Tests/RHITests/Source/RHITests/Public/RHITextureTests.h
+++ b/Engine/Plugins/Tests/RHITests/Source/RHITests/Public/RHITextureTests.h
@@ -1426,7 +1426,7 @@ public:
 				const FRHITextureCreateDesc SourceDesc =
 					FRHITextureCreateDesc::Create2D(TestName, Extent, PF_A2B10G10R10)
 					// Force as non-static
-					.AddFlags(ETextureCreateFlags::RenderTargetable);
+					.AddFlags(ETextureCreateFlags::RenderTargetable | ETextureCreateFlags::NoFastClear);
 
 				FTextureRHIRef SourceTexture = RHICmdList.CreateTexture(SourceDesc);
 
