Adding more FRHIDescriptorRange unit tests:
- Updating a subset of Descriptor Range resources.
- Use a Descriptor Range to store UAVs to write to.
- Use a Descriptor Range to store both SRVs to load from and a single UAV to write to.
- Use a Descriptor Range to store SRVs of both textures and buffers.

#Synced-CL 52447069

diff --git a/Engine/Plugins/Tests/RHITests/Source/RHITests/Private/RHIBindlessTests.cpp b/Engine/Plugins/Tests/RHITests/Source/RHITests/Private/RHIBindlessTests.cpp
--- a/Engine/Plugins/Tests/RHITests/Source/RHITests/Private/RHIBindlessTests.cpp
+++ b/Engine/Plugins/Tests/RHITests/Source/RHITests/Private/RHIBindlessTests.cpp
@@ -13,6 +13,19 @@
 #include "ShaderParameterMacros.h"
 #include "ShaderPlatformConfig.h"
 
+inline bool IsDescriptorTypeMaskSupported(ERHIDescriptorTypeMask TypeMask)
+{
+	for (ERHIDescriptorTypeMask SupportedGroup : GRHIGlobals.DescriptorTypeGroups)
+	{
+		if (EnumHasAllFlags(SupportedGroup, TypeMask))
+		{
+			return true;
+		}
+	}
+
+	return false;
+}
+
 bool RHIBindlessTests::Test_ResourceCollection(FRHICommandListImmediate& RHICmdList)
 {
 	if (!GRHIGlobals.bSupportsBindless || !IsBindlessEnabledForAnyGraphics(FShaderPlatformConfig::GetBindlessConfiguration(GMaxRHIShaderPlatform)))
@@ -74,10 +87,10 @@ bool RHIBindlessTests::Test_ResourceCollection(FRHICommandListImmediate& RHICmdL
 	return true;
 }
 
-class FDescriptorRangeFillFromTexturesCS : public FGlobalShader
+class FDescriptorRangeTestCS : public FGlobalShader
 {
-	DECLARE_GLOBAL_SHADER(FDescriptorRangeFillFromTexturesCS);
-	SHADER_USE_PARAMETER_STRUCT(FDescriptorRangeFillFromTexturesCS, FGlobalShader);
+public:
+	using FGlobalShader::FGlobalShader;
 
 	static constexpr uint32 ThreadGroupSizeX = 1;
 	static constexpr uint32 ThreadGroupSizeY = 1;
@@ -93,9 +106,111 @@ class FDescriptorRangeFillFromTexturesCS : public FGlobalShader
 		FGlobalShader::ModifyCompilationEnvironment(Parameters, OutEnvironment);
 		OutEnvironment.SetDefine(TEXT("THREADGROUP_SIZE_X"), ThreadGroupSizeX);
 		OutEnvironment.SetDefine(TEXT("THREADGROUP_SIZE_Y"), ThreadGroupSizeY);
-
 		OutEnvironment.CompilerFlags.Add(CFLAG_SupportsMinimalBindless);
 	}
+};
+
+static bool SupportsDescriptorRangeTests()
+{
+	return GRHIGlobals.bSupportsBindless
+		&& GRHIGlobals.bSupportsDescriptorRange
+		&& IsBindlessEnabledForAnyGraphics(FShaderPlatformConfig::GetBindlessConfiguration(GMaxRHIShaderPlatform));
+}
+
+// Helper: Create a 1x1 texture with a given color and return it along with its SRV.
+static void CreateColorTextureAndSRV(FRHICommandListBase& RHICmdList, FColor Color, FTextureRHIRef& OutTexture, FShaderResourceViewRHIRef& OutSRV)
+{
+	const FRHITextureCreateDesc TextureDesc =
+		FRHITextureCreateDesc::Create2D(TEXT("Test_DescriptorRange"), 1, 1, PF_R8G8B8A8)
+		.SetInitialState(ERHIAccess::SRVCompute)
+		.SetInitActionInitializer();
+
+	uint32 PackedColor = Color.ToPackedABGR();
+	FRHITextureInitializer Initializer = RHICmdList.CreateTextureInitializer(TextureDesc);
+	Initializer.GetTexture2DSubresource(0).WriteData(&PackedColor, sizeof(PackedColor));
+	OutTexture = Initializer.Finalize();
+
+	const FRHIViewDesc::FTextureSRV::FInitializer SRVInit =
+		FRHIViewDesc::CreateTextureSRV()
+		.SetDimension(TextureDesc.Dimension)
+		.SetFormat(PF_R8G8B8A8)
+		.SetMipRange(0, 1);
+
+	OutSRV = RHICmdList.CreateShaderResourceView(OutTexture, SRVInit);
+}
+
+// Helper: Create an array of 1x1 color textures with SRVs and descriptor resources.
+static void CreateColorTextureSRVResources(
+	FRHICommandListBase& RHICmdList,
+	TConstArrayView<FColor> Colors,
+	TArray<FTextureRHIRef>& OutTextures,
+	TArray<FShaderResourceViewRHIRef>& OutSRVs,
+	TArray<FRHIDescriptorResource>& OutResources)
+{
+	for (int32 Index = 0; Index < Colors.Num(); Index++)
+	{
+		FTextureRHIRef Texture;
+		FShaderResourceViewRHIRef SRV;
+		CreateColorTextureAndSRV(RHICmdList, Colors[Index], Texture, SRV);
+		OutResources.Emplace(SRV.GetReference());
+		OutTextures.Emplace(MoveTemp(Texture));
+		OutSRVs.Emplace(MoveTemp(SRV));
+	}
+}
+
+// Helper: Readback a structured buffer of float4s and compare against expected colors.
+static bool ReadbackAndVerifyColors(
+	FRHICommandListImmediate& RHICmdList,
+	FRHIBuffer* OutputBuffer,
+	uint32 OutputBufferSize,
+	int32 NumElements,
+	TConstArrayView<FLinearColor> ExpectedColors,
+	const TCHAR* TestName)
+{
+	FRHIGPUBufferReadback OutputBufferReadback(TEXT("OutputUAVReadback"));
+	OutputBufferReadback.EnqueueCopy(RHICmdList, OutputBuffer, OutputBufferSize);
+
+	RHICmdList.BlockUntilGPUIdle();
+	RHICmdList.FlushResources();
+
+	bool bExpectedResult = true;
+
+	{
+		const FLinearColor* OutputData = reinterpret_cast<const FLinearColor*>(OutputBufferReadback.Lock(OutputBufferSize));
+		for (int32 Index = 0; Index < NumElements; Index++)
+		{
+			const FLinearColor& ExpectedColor = ExpectedColors[Index];
+			const FLinearColor& ActualColor = OutputData[Index];
+
+			if (ExpectedColor != ActualColor)
+			{
+				UE_LOG(LogRHIUnitTestCommandlet, Error, TEXT("Test failed. \"%s\"\n\nIndex: %d;"
+					" Expected Data: (%.2f, %.2f, %.2f, %.2f); Actual Data: (%.2f, %.2f, %.2f, %.2f);"),
+					TestName, Index,
+					ExpectedColor.R, ExpectedColor.G, ExpectedColor.B, ExpectedColor.A,
+					ActualColor.R, ActualColor.G, ActualColor.B, ActualColor.A);
+
+				bExpectedResult = false;
+				break;
+			}
+		}
+		OutputBufferReadback.Unlock();
+
+		if (bExpectedResult)
+		{
+			UE_LOG(LogRHIUnitTestCommandlet, Display, TEXT("Test passed. \"%s\""), TestName);
+		}
+	}
+
+	return bExpectedResult;
+}
+
+// Helper: Dispatch FillFromTextures, readback, and compare against expected colors.
+
+class FDescriptorRangeFillFromTexturesCS : public FDescriptorRangeTestCS
+{
+	DECLARE_GLOBAL_SHADER(FDescriptorRangeFillFromTexturesCS);
+	SHADER_USE_PARAMETER_STRUCT(FDescriptorRangeFillFromTexturesCS, FDescriptorRangeTestCS);
 
 	BEGIN_SHADER_PARAMETER_STRUCT(FParameters, )
 		SHADER_PARAMETER_UAV(RWBuffer<float4>, OutputUAV)
@@ -104,74 +219,322 @@ class FDescriptorRangeFillFromTexturesCS : public FGlobalShader
 };
 IMPLEMENT_GLOBAL_SHADER(FDescriptorRangeFillFromTexturesCS, "/Plugin/RHITests/Private/TestDescriptorRange.usf", "FillFromTextures", SF_Compute);
 
-bool RHIBindlessTests::Test_DescriptorRange(FRHICommandListImmediate& RHICmdList)
+static bool DispatchAndVerifyDescriptorRange(
+	FRHICommandListImmediate& RHICmdList,
+	FDescriptorRangeRHIRef& DescriptorRange,
+	int32 NumResources,
+	TConstArrayView<FLinearColor> ExpectedColors,
+	const TCHAR* TestName)
+{
+	FRHIBufferCreateDesc OutputBufferDesc =
+		FRHIBufferCreateDesc::CreateStructured<FLinearColor>(TEXT("OutputUAV"), NumResources)
+		.AddUsage(EBufferUsageFlags::UnorderedAccess | EBufferUsageFlags::SourceCopy)
+		.SetInitialState(ERHIAccess::UAVCompute);
+	FBufferRHIRef OutputBuffer = RHICmdList.CreateBuffer(OutputBufferDesc);
+
+	FUnorderedAccessViewRHIRef OutputUAV = RHICmdList.CreateUnorderedAccessView(
+		OutputBuffer,
+		FRHIViewDesc::CreateBufferUAV()
+			.SetType(FRHIViewDesc::EBufferType::Structured)
+			.SetStride(OutputBufferDesc.Stride).SetNumElements(NumResources)
+	);
+
+	{
+		TShaderMapRef<FDescriptorRangeFillFromTexturesCS> ComputeShader(GetGlobalShaderMap(GMaxRHIFeatureLevel));
+		const uint32 NumThreadGroups = FMath::DivideAndRoundUp<uint32>(NumResources, FDescriptorRangeFillFromTexturesCS::ThreadGroupSizeX);
+
+		FDescriptorRangeFillFromTexturesCS::FParameters Parameters;
+		Parameters.OutputUAV = OutputUAV;
+		Parameters.DescriptorRange = DescriptorRange;
+
+		FComputeShaderUtils::Dispatch(RHICmdList, ComputeShader, Parameters, FIntVector(NumThreadGroups, 1, 1));
+		RHICmdList.Transition(FRHITransitionInfo(OutputUAV, ERHIAccess::UAVCompute, ERHIAccess::CopySrc));
+	}
+
+	return ReadbackAndVerifyColors(RHICmdList, OutputBuffer, OutputBufferDesc.Size, NumResources, ExpectedColors, TestName);
+}
+
+// Test reading from a descriptor range of texture SRVs.
+// Creates SRVs from 1x1 color textures, reads them via a descriptor range, and verifies the output matches.
+bool RHIBindlessTests::Test_DescriptorRange_SRV(FRHICommandListImmediate& RHICmdList)
 {
-	if (!GRHIGlobals.bSupportsBindless || !GRHIGlobals.bSupportsDescriptorRange || !IsBindlessEnabledForAnyGraphics(FShaderPlatformConfig::GetBindlessConfiguration(GMaxRHIShaderPlatform)))
+	if (!SupportsDescriptorRangeTests())
 	{
 		return true;
 	}
 
-	const FColor TextureColors[6] =
+	const FColor TextureColors[] = { FColor::White, FColor::Black, FColor::Red, FColor::Green, FColor::Blue, FColor::Yellow };
+	constexpr int32 NumResources = UE_ARRAY_COUNT(TextureColors);
+
+	TArray<FTextureRHIRef> Textures;
+	TArray<FShaderResourceViewRHIRef> SRVs;
+	TArray<FRHIDescriptorResource> Resources;
+
+	CreateColorTextureSRVResources(RHICmdList, TextureColors, Textures, SRVs, Resources);
+
+	FRHIDescriptorRangeDesc RangeDesc;
+	RangeDesc.Count = NumResources;
+	RangeDesc.TypeMask = ERHIDescriptorTypeMask::SRV;
+
+	FDescriptorRangeRHIRef DescriptorRange = RHICmdList.CreateDescriptorRange(RangeDesc, Resources);
+	if (!DescriptorRange)
+	{
+		UE_LOG(LogRHIUnitTestCommandlet, Error, TEXT("Failed to allocate RHI Descriptor Range."));
+		return false;
+	}
+
+	TArray<FLinearColor> ExpectedColors;
+	for (int32 Index = 0; Index < NumResources; Index++)
 	{
-		FColor::White,
-		FColor::Black,
-		FColor::Red,
-		FColor::Green,
-		FColor::Blue,
-		FColor::Yellow,
-	};
+		ExpectedColors.Emplace(FLinearColor(TextureColors[Index]));
+	}
+
+	return DispatchAndVerifyDescriptorRange(RHICmdList, DescriptorRange, NumResources, ExpectedColors, TEXT("Descriptor Range SRV"));
+}
+
+
+// Test updating a subset of slots in an existing descriptor range.
+bool RHIBindlessTests::Test_DescriptorRange_UpdateWithOffset(FRHICommandListImmediate& RHICmdList)
+{
+	if (!SupportsDescriptorRangeTests())
+	{
+		return true;
+	}
+
+	const FColor InitialColor = FColor::White;
+	const FColor UpdatedColors[] = { FColor::Red, FColor::Green, FColor::Blue };
+
+	constexpr int32 NumSlots = 6;
+	constexpr int32 UpdateStartIndex = 2;
+	constexpr int32 NumUpdated = UE_ARRAY_COUNT(UpdatedColors);
+
+	// Create initial textures — all white.
+	TArray<FColor> InitialColors;
+	InitialColors.SetNum(NumSlots);
+	for (FColor& Color : InitialColors) { Color = InitialColor; }
 
 	TArray<FTextureRHIRef> Textures;
 	TArray<FShaderResourceViewRHIRef> SRVs;
+	TArray<FRHIDescriptorResource> Resources;
+
+	CreateColorTextureSRVResources(RHICmdList, InitialColors, Textures, SRVs, Resources);
+
+	FRHIDescriptorRangeDesc RangeDesc;
+	RangeDesc.Count = NumSlots;
+	RangeDesc.TypeMask = ERHIDescriptorTypeMask::SRV;
 
+	FDescriptorRangeRHIRef DescriptorRange = RHICmdList.CreateDescriptorRange(RangeDesc, Resources);
+	if (!DescriptorRange)
 	{
-		const FRHITextureCreateDesc TextureDesc =
-			FRHITextureCreateDesc::Create2D(TEXT("Test_DescriptorRange"), 1, 1, PF_R8G8B8A8)
-			.SetInitialState(ERHIAccess::SRVCompute)
-			.SetInitActionInitializer();
+		UE_LOG(LogRHIUnitTestCommandlet, Error, TEXT("Failed to allocate RHI Descriptor Range."));
+		return false;
+	}
 
-		for (FColor TextureColor : TextureColors)
-		{
-			uint32 Color = TextureColor.ToPackedABGR();
+	// Update slots 2, 3, 4 with new colors.
+	TArray<FTextureRHIRef> UpdatedTextures;
+	TArray<FShaderResourceViewRHIRef> UpdatedSRVs;
+	TArray<FRHIDescriptorResource> UpdatedResources;
 
-			FRHITextureInitializer Initializer = RHICmdList.CreateTextureInitializer(TextureDesc);
-			Initializer.GetTexture2DSubresource(0).WriteData(&Color, sizeof(Color));
-			Textures.Emplace(Initializer.Finalize());
-		}
+	CreateColorTextureSRVResources(RHICmdList, UpdatedColors, UpdatedTextures, UpdatedSRVs, UpdatedResources);
 
-		const FRHIViewDesc::FTextureSRV::FInitializer SRVInit =
-			FRHIViewDesc::CreateTextureSRV()
-			.SetDimension(TextureDesc.Dimension)
-			.SetFormat(PF_R8G8B8A8)
-			.SetMipRange(0, 1);
+	RHICmdList.UpdateDescriptorRange(DescriptorRange, UpdateStartIndex, UpdatedResources);
 
-		for (const FTextureRHIRef& Texture : Textures)
+	// Expected: White, White, Red, Green, Blue, White
+	TArray<FLinearColor> ExpectedColors;
+	for (int32 Index = 0; Index < NumSlots; Index++)
+	{
+		if (Index >= UpdateStartIndex && Index < UpdateStartIndex + NumUpdated)
 		{
-			SRVs.Emplace(RHICmdList.CreateShaderResourceView(Texture, SRVInit));
+			ExpectedColors.Emplace(FLinearColor(UpdatedColors[Index - UpdateStartIndex]));
+		}
+		else
+		{
+			ExpectedColors.Emplace(FLinearColor(InitialColor));
 		}
 	}
 
-	FRHIDescriptorRangeDesc RangeDesc;
-	RangeDesc.Count = Textures.Num();
-	RangeDesc.TypeMask = ERHIDescriptorTypeMask::SRV;
+	return DispatchAndVerifyDescriptorRange(RHICmdList, DescriptorRange, NumSlots, ExpectedColors, TEXT("Descriptor Range Update With Offset"));
+}
+
+class FDescriptorRangeWriteToUAVsCS : public FDescriptorRangeTestCS
+{
+	DECLARE_GLOBAL_SHADER(FDescriptorRangeWriteToUAVsCS);
+	SHADER_USE_PARAMETER_STRUCT(FDescriptorRangeWriteToUAVsCS, FDescriptorRangeTestCS);
+
+	BEGIN_SHADER_PARAMETER_STRUCT(FParameters, )
+		SHADER_PARAMETER_SRV(StructuredBuffer<float4>, InputSRV)
+		SHADER_PARAMETER_DESCRIPTOR_RANGE(FDescriptorRange, DescriptorRange)
+	END_SHADER_PARAMETER_STRUCT()
+};
+IMPLEMENT_GLOBAL_SHADER(FDescriptorRangeWriteToUAVsCS, "/Plugin/RHITests/Private/TestDescriptorRange.usf", "WriteToUAVs", SF_Compute);
+
+// Test writing to a descriptor range of UAVs.
+bool RHIBindlessTests::Test_DescriptorRange_UAV(FRHICommandListImmediate& RHICmdList)
+{
+	if (!SupportsDescriptorRangeTests())
+	{
+		return true;
+	}
+
+	constexpr int32 NumResources = 4;
 
+	const FLinearColor InputValues[NumResources] =
+	{
+		FLinearColor(1.0f, 0.0f, 0.0f, 1.0f),
+		FLinearColor(0.0f, 1.0f, 0.0f, 1.0f),
+		FLinearColor(0.0f, 0.0f, 1.0f, 1.0f),
+		FLinearColor(1.0f, 1.0f, 0.0f, 1.0f),
+	};
+
+	// Create an input buffer with the values we want to write through the UAV range.
+	FRHIBufferCreateDesc InputBufferDesc =
+		FRHIBufferCreateDesc::CreateStructured<FLinearColor>(TEXT("InputSRV"), NumResources)
+		.SetInitialState(ERHIAccess::SRVCompute)
+		.SetInitActionInitializer();
+	TRHIBufferInitializer<FLinearColor> InputInitializer = RHICmdList.CreateBufferInitializer(InputBufferDesc);
+	{
+		InputInitializer.WriteArray(MakeConstArrayView(InputValues));
+	}
+	FBufferRHIRef InputBuffer = InputInitializer.Finalize();
+
+	FShaderResourceViewRHIRef InputSRV = RHICmdList.CreateShaderResourceView(
+		InputBuffer,
+		FRHIViewDesc::CreateBufferSRV()
+			.SetType(FRHIViewDesc::EBufferType::Structured)
+			.SetStride(sizeof(FLinearColor)).SetNumElements(NumResources)
+	);
+
+	// Create output UAV buffers — one per descriptor range slot.
+	TArray<FBufferRHIRef> OutputBuffers;
+	TArray<FUnorderedAccessViewRHIRef> OutputUAVs;
 	TArray<FRHIDescriptorResource> Resources;
-	for (const FShaderResourceViewRHIRef& SRV : SRVs)
+
+	for (int32 Index = 0; Index < NumResources; Index++)
 	{
-		Resources.Emplace(SRV.GetReference());
+		FRHIBufferCreateDesc OutputBufferDesc =
+			FRHIBufferCreateDesc::CreateStructured<FLinearColor>(TEXT("UAVRangeOutput"), 1)
+			.AddUsage(EBufferUsageFlags::UnorderedAccess | EBufferUsageFlags::SourceCopy)
+			.SetInitialState(ERHIAccess::UAVCompute);
+		FBufferRHIRef OutputBuffer = RHICmdList.CreateBuffer(OutputBufferDesc);
+
+		FUnorderedAccessViewRHIRef OutputUAV = RHICmdList.CreateUnorderedAccessView(
+			OutputBuffer,
+			FRHIViewDesc::CreateBufferUAV()
+				.SetType(FRHIViewDesc::EBufferType::Structured)
+				.SetStride(sizeof(FLinearColor)).SetNumElements(1)
+		);
+
+		Resources.Emplace(OutputUAV.GetReference());
+		OutputBuffers.Emplace(MoveTemp(OutputBuffer));
+		OutputUAVs.Emplace(MoveTemp(OutputUAV));
 	}
 
-	FDescriptorRangeRHIRef DescriptorRange = RHICmdList.CreateDescriptorRange(RangeDesc, Resources);
+	FRHIDescriptorRangeDesc RangeDesc;
+	RangeDesc.Count = NumResources;
+	RangeDesc.TypeMask = ERHIDescriptorTypeMask::UAV;
 
+	FDescriptorRangeRHIRef DescriptorRange = RHICmdList.CreateDescriptorRange(RangeDesc, Resources);
 	if (!DescriptorRange)
 	{
-		UE_LOG(LogRHIUnitTestCommandlet, Display, TEXT("Unable to allocate RHI Descriptor Range. Assuming it is not implemented for the current platform."));
+		UE_LOG(LogRHIUnitTestCommandlet, Error, TEXT("Failed to allocate RHI Descriptor Range for UAVs."));
+		return false;
+	}
+
+	// Dispatch the WriteToUAVs shader.
+	{
+		TShaderMapRef<FDescriptorRangeWriteToUAVsCS> ComputeShader(GetGlobalShaderMap(GMaxRHIFeatureLevel));
+
+		FDescriptorRangeWriteToUAVsCS::FParameters Parameters;
+		Parameters.InputSRV = InputSRV;
+		Parameters.DescriptorRange = DescriptorRange;
+
+		FComputeShaderUtils::Dispatch(RHICmdList, ComputeShader, Parameters, FIntVector(NumResources, 1, 1));
+
+		for (const FUnorderedAccessViewRHIRef& UAV : OutputUAVs)
+		{
+			RHICmdList.Transition(FRHITransitionInfo(UAV, ERHIAccess::UAVCompute, ERHIAccess::CopySrc));
+		}
+	}
+
+	RHICmdList.BlockUntilGPUIdle();
+	RHICmdList.FlushResources();
+
+	// Readback each output buffer and verify.
+	bool bExpectedResult = true;
+
+	for (int32 Index = 0; Index < NumResources; Index++)
+	{
+		FRHIGPUBufferReadback Readback(TEXT("UAVRangeReadback"));
+		Readback.EnqueueCopy(RHICmdList, OutputBuffers[Index], sizeof(FLinearColor));
+		RHICmdList.BlockUntilGPUIdle();
+		RHICmdList.FlushResources();
+
+		const FLinearColor* OutputData = reinterpret_cast<const FLinearColor*>(Readback.Lock(sizeof(FLinearColor)));
+		const FLinearColor& ExpectedColor = InputValues[Index];
+		const FLinearColor& ActualColor = *OutputData;
+
+		if (ExpectedColor != ActualColor)
+		{
+			UE_LOG(LogRHIUnitTestCommandlet, Error, TEXT("Test failed. \"Descriptor Range UAV\"\n\nIndex: %d;"
+				" Expected Data: (%.2f, %.2f, %.2f, %.2f); Actual Data: (%.2f, %.2f, %.2f, %.2f);"),
+				Index,
+				ExpectedColor.R, ExpectedColor.G, ExpectedColor.B, ExpectedColor.A,
+				ActualColor.R, ActualColor.G, ActualColor.B, ActualColor.A);
+
+			bExpectedResult = false;
+		}
+		Readback.Unlock();
+
+		if (!bExpectedResult)
+		{
+			break;
+		}
+	}
+
+	if (bExpectedResult)
+	{
+		UE_LOG(LogRHIUnitTestCommandlet, Display, TEXT("Test passed. \"Descriptor Range UAV\""));
+	}
+
+	return bExpectedResult;
+}
+
+class FDescriptorRangeReadSRVsWriteUAVCS : public FDescriptorRangeTestCS
+{
+	DECLARE_GLOBAL_SHADER(FDescriptorRangeReadSRVsWriteUAVCS);
+	SHADER_USE_PARAMETER_STRUCT(FDescriptorRangeReadSRVsWriteUAVCS, FDescriptorRangeTestCS);
+
+	BEGIN_SHADER_PARAMETER_STRUCT(FParameters, )
+		SHADER_PARAMETER_DESCRIPTOR_RANGE(FDescriptorRange, DescriptorRange)
+	END_SHADER_PARAMETER_STRUCT()
+};
+IMPLEMENT_GLOBAL_SHADER(FDescriptorRangeReadSRVsWriteUAVCS, "/Plugin/RHITests/Private/TestDescriptorRange.usf", "ReadSRVsWriteUAV", SF_Compute);
+
+// Test a descriptor range containing both SRV and UAV resources.
+// Places a single UAV at index 0 followed by texture SRVs. The shader reads each SRV and writes to the UAV, verifying both view types work in the same range.
+bool RHIBindlessTests::Test_DescriptorRange_SRV_And_UAV(FRHICommandListImmediate& RHICmdList)
+{
+	if (!SupportsDescriptorRangeTests())
+	{
 		return true;
 	}
 
+	if (!IsDescriptorTypeMaskSupported(ERHIDescriptorTypeMask::SRV | ERHIDescriptorTypeMask::UAV))
+	{
+		UE_LOG(LogRHIUnitTestCommandlet, Display, TEXT("SRV+UAV Descriptor Range is not supported on this platform. Skipping test."));
+		return true;
+	}
+
+	// Test a descriptor range containing a single UAV followed by SRVs: [UAV, SRV0, SRV1, SRV2, SRV3]
+	// The shader reads from each SRV and writes the result into the UAV at the corresponding index.
+	const FColor TextureColors[] = { FColor::Red, FColor::Green, FColor::Blue, FColor::Yellow };
+	constexpr int32 NumTextures = UE_ARRAY_COUNT(TextureColors);
+	constexpr int32 NumResources = 1 + NumTextures; // 1 UAV + N SRVs
+
+	// Create the output UAV buffer — holds one float4 per texture.
 	FRHIBufferCreateDesc OutputBufferDesc =
-		FRHIBufferCreateDesc::CreateStructured<FLinearColor>(TEXT("OutputUAV"), Resources.Num())
-		.AddUsage(EBufferUsageFlags::UnorderedAccess)
+		FRHIBufferCreateDesc::CreateStructured<FLinearColor>(TEXT("SRVAndUAVRangeOutput"), NumTextures)
+		.AddUsage(EBufferUsageFlags::UnorderedAccess | EBufferUsageFlags::SourceCopy)
 		.SetInitialState(ERHIAccess::UAVCompute);
 	FBufferRHIRef OutputBuffer = RHICmdList.CreateBuffer(OutputBufferDesc);
 
@@ -179,61 +542,162 @@ bool RHIBindlessTests::Test_DescriptorRange(FRHICommandListImmediate& RHICmdList
 		OutputBuffer,
 		FRHIViewDesc::CreateBufferUAV()
 			.SetType(FRHIViewDesc::EBufferType::Structured)
-			.SetStride(OutputBufferDesc.Stride).SetNumElements(Resources.Num())
+			.SetStride(sizeof(FLinearColor)).SetNumElements(NumTextures)
 	);
 
+	// Build the resource list: UAV at index 0, then SRVs.
+	TArray<FTextureRHIRef> Textures;
+	TArray<FShaderResourceViewRHIRef> SRVs;
+	TArray<FRHIDescriptorResource> Resources;
+
+	Resources.Emplace(OutputUAV.GetReference());
+
+	CreateColorTextureSRVResources(RHICmdList, TextureColors, Textures, SRVs, Resources);
+
+	FRHIDescriptorRangeDesc RangeDesc;
+	RangeDesc.Count = NumResources;
+	RangeDesc.TypeMask = ERHIDescriptorTypeMask::SRV | ERHIDescriptorTypeMask::UAV;
+
+	FDescriptorRangeRHIRef DescriptorRange = RHICmdList.CreateDescriptorRange(RangeDesc, Resources);
+	if (!DescriptorRange)
 	{
-		TShaderMapRef< FDescriptorRangeFillFromTexturesCS> ComputeShader(GetGlobalShaderMap(GMaxRHIFeatureLevel));
-		const uint32 NumThreadGroups = FMath::DivideAndRoundUp<uint32>(Resources.Num(), FDescriptorRangeFillFromTexturesCS::ThreadGroupSizeX);
+		UE_LOG(LogRHIUnitTestCommandlet, Display, TEXT("Unable to allocate RHI Descriptor Range for SRV+UAV."));
+		return false;
+	}
 
-		FDescriptorRangeFillFromTexturesCS::FParameters Parameters;
-		Parameters.OutputUAV = OutputUAV;
+	// Dispatch the ReadSRVsWriteUAV shader.
+	{
+		TShaderMapRef<FDescriptorRangeReadSRVsWriteUAVCS> ComputeShader(GetGlobalShaderMap(GMaxRHIFeatureLevel));
+
+		FDescriptorRangeReadSRVsWriteUAVCS::FParameters Parameters;
 		Parameters.DescriptorRange = DescriptorRange;
-		//Parameters.DescriptorRange_StartIndex = DescriptorRange->GetShaderData().StartIndex;
-		//Parameters.DescriptorRange_Count = DescriptorRange->GetShaderData().Count;
-		
-		FComputeShaderUtils::Dispatch(RHICmdList, ComputeShader, Parameters, FIntVector(NumThreadGroups, 1, 1));
+
+		FComputeShaderUtils::Dispatch(RHICmdList, ComputeShader, Parameters, FIntVector(NumTextures, 1, 1));
 
 		RHICmdList.Transition(FRHITransitionInfo(OutputUAV, ERHIAccess::UAVCompute, ERHIAccess::CopySrc));
 	}
 
-	FRHIGPUBufferReadback OutputBufferReadback(TEXT("OutputUAVReadback"));
-	OutputBufferReadback.EnqueueCopy(RHICmdList, OutputBuffer, OutputBufferDesc.Size);
+	TArray<FLinearColor> ExpectedColors;
+	for (int32 Index = 0; Index < NumTextures; Index++)
+	{
+		ExpectedColors.Emplace(FLinearColor(TextureColors[Index]));
+	}
 
-	// Sync the GPU. Unfortunately we can't use the fences because not all RHIs implement them yet.
-	RHICmdList.BlockUntilGPUIdle();
-	RHICmdList.FlushResources();
+	return ReadbackAndVerifyColors(RHICmdList, OutputBuffer, OutputBufferDesc.Size, NumTextures, ExpectedColors, TEXT("Descriptor Range SRV And UAV"));
+}
 
-	bool bExpectedResult = true;
+class FDescriptorRangeFillFromMixedResourcesCS : public FDescriptorRangeTestCS
+{
+	DECLARE_GLOBAL_SHADER(FDescriptorRangeFillFromMixedResourcesCS);
+	SHADER_USE_PARAMETER_STRUCT(FDescriptorRangeFillFromMixedResourcesCS, FDescriptorRangeTestCS);
 
+	BEGIN_SHADER_PARAMETER_STRUCT(FParameters, )
+		SHADER_PARAMETER_UAV(RWBuffer<float4>, OutputUAV)
+		SHADER_PARAMETER_DESCRIPTOR_RANGE(FDescriptorRange, DescriptorRange)
+	END_SHADER_PARAMETER_STRUCT()
+};
+IMPLEMENT_GLOBAL_SHADER(FDescriptorRangeFillFromMixedResourcesCS, "/Plugin/RHITests/Private/TestDescriptorRange.usf", "FillFromMixedResources", SF_Compute);
+
+// Test a descriptor range containing a mix of texture SRVs and buffer SRVs.
+bool RHIBindlessTests::Test_DescriptorRange_MixedResourceTypes(FRHICommandListImmediate& RHICmdList)
+{
+	if (!SupportsDescriptorRangeTests())
 	{
-		const FLinearColor* OutputData = reinterpret_cast<const FLinearColor*>(OutputBufferReadback.Lock(OutputBufferDesc.Size));
-		for (int32 Index = 0; Index < Resources.Num(); Index++)
-		{
-			const FLinearColor ExpectedColor = FLinearColor(TextureColors[Index]);
-			const FLinearColor ActualColor = OutputData[Index];
+		return true;
+	}
 
-			if (ExpectedColor != ActualColor)
-			{
-				UE_LOG(LogRHIUnitTestCommandlet, Error, TEXT("Test failed. \"Descriptor Range\"\n\nIndex: %d;"
-					" Expected Data: (%.2f, %.2f, %.2f, %.2f); Actual Data: (%.2f, %.2f, %.2f, %.2f);"),
-					Index,
-					ExpectedColor.R, ExpectedColor.G, ExpectedColor.B, ExpectedColor.A,
-					ActualColor.R, ActualColor.G, ActualColor.B, ActualColor.A
-					);
+	// Mix of textures and buffers in the same range, interleaved as pairs: [Tex0, Buf0, Tex1, Buf1, ...]
+	// The shader reads from each texture and buffer and writes both to the output.
+	const FColor Colors[] = { FColor::Red, FColor::Green, FColor::Blue, FColor::Yellow };
+	constexpr int32 NumPairs = UE_ARRAY_COUNT(Colors);
+	constexpr int32 NumResources = NumPairs * 2; // Texture + Buffer per pair
 
-				bExpectedResult = false;
-				break;
-			}
-		}
-		OutputBufferReadback.Unlock();
+	TArray<FTextureRHIRef> Textures;
+	TArray<FShaderResourceViewRHIRef> TextureSRVs;
+	TArray<FBufferRHIRef> Buffers;
+	TArray<FShaderResourceViewRHIRef> BufferSRVs;
+	TArray<FRHIDescriptorResource> Resources;
 
-		if (bExpectedResult)
+	for (int32 Index = 0; Index < NumPairs; Index++)
+	{
+		// Even slot: 1x1 texture SRV
+		FTextureRHIRef Texture;
+		FShaderResourceViewRHIRef TextureSRV;
+		CreateColorTextureAndSRV(RHICmdList, Colors[Index], Texture, TextureSRV);
+
+		// Odd slot: buffer SRV containing the same color
+		FLinearColor BufferColor(Colors[Index]);
+		FRHIBufferCreateDesc BufferDesc =
+			FRHIBufferCreateDesc::Create(TEXT("MixedRangeBuffer"), sizeof(FLinearColor), sizeof(FLinearColor), EBufferUsageFlags::ShaderResource)
+			.SetInitialState(ERHIAccess::SRVCompute)
+			.SetInitActionInitializer();
+		TRHIBufferInitializer<FLinearColor> BufferInitializer = RHICmdList.CreateBufferInitializer(BufferDesc);
 		{
-			UE_LOG(LogRHIUnitTestCommandlet, Display, TEXT("Test passed. \"Descriptor Range\""));
+			BufferInitializer.WriteData(&BufferColor, sizeof(BufferColor));
 		}
+		FBufferRHIRef Buffer = BufferInitializer.Finalize();
+
+		FShaderResourceViewRHIRef BufferSRV = RHICmdList.CreateShaderResourceView(
+			Buffer,
+			FRHIViewDesc::CreateBufferSRV()
+				.SetType(FRHIViewDesc::EBufferType::Typed)
+				.SetFormat(PF_A32B32G32R32F)
+		);
+
+		Resources.Emplace(TextureSRV.GetReference());
+		Resources.Emplace(BufferSRV.GetReference());
+
+		Textures.Emplace(MoveTemp(Texture));
+		TextureSRVs.Emplace(MoveTemp(TextureSRV));
+		Buffers.Emplace(MoveTemp(Buffer));
+		BufferSRVs.Emplace(MoveTemp(BufferSRV));
 	}
 
-	return bExpectedResult;
+	FRHIDescriptorRangeDesc RangeDesc;
+	RangeDesc.Count = NumResources;
+	RangeDesc.TypeMask = ERHIDescriptorTypeMask::SRV;
+
+	FDescriptorRangeRHIRef DescriptorRange = RHICmdList.CreateDescriptorRange(RangeDesc, Resources);
+	if (!DescriptorRange)
+	{
+		UE_LOG(LogRHIUnitTestCommandlet, Error, TEXT("Failed to allocate RHI Descriptor Range."));
+		return false;
+	}
+
+	// Create output buffer and UAV.
+	FRHIBufferCreateDesc OutputBufferDesc =
+		FRHIBufferCreateDesc::CreateStructured<FLinearColor>(TEXT("OutputUAV"), NumResources)
+		.AddUsage(EBufferUsageFlags::UnorderedAccess | EBufferUsageFlags::SourceCopy)
+		.SetInitialState(ERHIAccess::UAVCompute);
+	FBufferRHIRef OutputBuffer = RHICmdList.CreateBuffer(OutputBufferDesc);
+
+	FUnorderedAccessViewRHIRef OutputUAV = RHICmdList.CreateUnorderedAccessView(
+		OutputBuffer,
+		FRHIViewDesc::CreateBufferUAV()
+			.SetType(FRHIViewDesc::EBufferType::Structured)
+			.SetStride(OutputBufferDesc.Stride).SetNumElements(NumResources)
+	);
+
+	{
+		TShaderMapRef<FDescriptorRangeFillFromMixedResourcesCS> ComputeShader(GetGlobalShaderMap(GMaxRHIFeatureLevel));
+
+		FDescriptorRangeFillFromMixedResourcesCS::FParameters Parameters;
+		Parameters.OutputUAV = OutputUAV;
+		Parameters.DescriptorRange = DescriptorRange;
+
+		FComputeShaderUtils::Dispatch(RHICmdList, ComputeShader, Parameters, FIntVector(NumPairs, 1, 1));
+		RHICmdList.Transition(FRHITransitionInfo(OutputUAV, ERHIAccess::UAVCompute, ERHIAccess::CopySrc));
+	}
+
+	// Expected: each color appears twice (once from texture, once from buffer).
+	TArray<FLinearColor> ExpectedColors;
+	for (int32 Index = 0; Index < NumPairs; Index++)
+	{
+		FLinearColor Color(Colors[Index]);
+		ExpectedColors.Emplace(Color); // from texture
+		ExpectedColors.Emplace(Color); // from buffer
+	}
+
+	return ReadbackAndVerifyColors(RHICmdList, OutputBuffer, OutputBufferDesc.Size, NumResources, ExpectedColors, TEXT("Descriptor Range Mixed Resource Types"));
 }
 
diff --git a/Engine/Plugins/Tests/RHITests/Source/RHITests/Private/RHIBindlessTests.h b/Engine/Plugins/Tests/RHITests/Source/RHITests/Private/RHIBindlessTests.h
--- a/Engine/Plugins/Tests/RHITests/Source/RHITests/Private/RHIBindlessTests.h
+++ b/Engine/Plugins/Tests/RHITests/Source/RHITests/Private/RHIBindlessTests.h
@@ -7,5 +7,9 @@ class FRHICommandListImmediate;
 namespace RHIBindlessTests
 {
 	bool Test_ResourceCollection(FRHICommandListImmediate& RHICmdList);
-	bool Test_DescriptorRange(FRHICommandListImmediate& RHICmdList);
+	bool Test_DescriptorRange_SRV(FRHICommandListImmediate& RHICmdList);
+	bool Test_DescriptorRange_UpdateWithOffset(FRHICommandListImmediate& RHICmdList);
+	bool Test_DescriptorRange_UAV(FRHICommandListImmediate& RHICmdList);
+	bool Test_DescriptorRange_SRV_And_UAV(FRHICommandListImmediate& RHICmdList);
+	bool Test_DescriptorRange_MixedResourceTypes(FRHICommandListImmediate& RHICmdList);
 }
diff --git a/Engine/Plugins/Tests/RHITests/Source/RHITests/Private/RHITests.spec.cpp b/Engine/Plugins/Tests/RHITests/Source/RHITests/Private/RHITests.spec.cpp
--- a/Engine/Plugins/Tests/RHITests/Source/RHITests/Private/RHITests.spec.cpp
+++ b/Engine/Plugins/Tests/RHITests/Source/RHITests/Private/RHITests.spec.cpp
@@ -219,10 +219,30 @@ void FAutomationRHITest::Define()
 			bool bResult = RunOnRenderThreadSynchronous(RHIBindlessTests::Test_ResourceCollection);
 			TestEqual("Resource Collection failed", bResult, 1);
 		});
-		It("RHI Descriptor Range", [this]()
+		It("RHI Descriptor Range SRV", [this]()
 		{
-			bool bResult = RunOnRenderThreadSynchronous(RHIBindlessTests::Test_DescriptorRange);
-			TestEqual("Descriptor Range failed", bResult, 1);
+			bool bResult = RunOnRenderThreadSynchronous(RHIBindlessTests::Test_DescriptorRange_SRV);
+			TestEqual("Descriptor Range SRV failed", bResult, 1);
+		});
+		It("RHI Descriptor Range Update With Offset", [this]()
+		{
+			bool bResult = RunOnRenderThreadSynchronous(RHIBindlessTests::Test_DescriptorRange_UpdateWithOffset);
+			TestEqual("Descriptor Range Update With Offset failed", bResult, 1);
+		});
+		It("RHI Descriptor Range UAV", [this]()
+		{
+			bool bResult = RunOnRenderThreadSynchronous(RHIBindlessTests::Test_DescriptorRange_UAV);
+			TestEqual("Descriptor Range UAV failed", bResult, 1);
+		});
+		It("RHI Descriptor Range SRV And UAV", [this]()
+		{
+			bool bResult = RunOnRenderThreadSynchronous(RHIBindlessTests::Test_DescriptorRange_SRV_And_UAV);
+			TestEqual("Descriptor Range SRV And UAV failed", bResult, 1);
+		});
+		It("RHI Descriptor Range Mixed Resource Types", [this]()
+		{
+			bool bResult = RunOnRenderThreadSynchronous(RHIBindlessTests::Test_DescriptorRange_MixedResourceTypes);
+			TestEqual("Descriptor Range Mixed Resource Types failed", bResult, 1);
 		});
 	});
 }
diff --git a/Engine/Plugins/Tests/RHITests/Source/RHITests/Private/RHITestsModule.cpp b/Engine/Plugins/Tests/RHITests/Source/RHITests/Private/RHITestsModule.cpp
--- a/Engine/Plugins/Tests/RHITests/Source/RHITests/Private/RHITestsModule.cpp
+++ b/Engine/Plugins/Tests/RHITests/Source/RHITests/Private/RHITestsModule.cpp
@@ -126,7 +126,11 @@ static bool RunTests_RenderThread(FRHICommandListImmediate& RHICmdList)
 	// ------------------------------------------------
 	{
 		RUN_TEST(RHIBindlessTests::Test_ResourceCollection(RHICmdList));
-		RUN_TEST(RHIBindlessTests::Test_DescriptorRange(RHICmdList));
+		RUN_TEST(RHIBindlessTests::Test_DescriptorRange_SRV(RHICmdList));
+		RUN_TEST(RHIBindlessTests::Test_DescriptorRange_UpdateWithOffset(RHICmdList));
+		RUN_TEST(RHIBindlessTests::Test_DescriptorRange_UAV(RHICmdList));
+		RUN_TEST(RHIBindlessTests::Test_DescriptorRange_SRV_And_UAV(RHICmdList));
+		RUN_TEST(RHIBindlessTests::Test_DescriptorRange_MixedResourceTypes(RHICmdList));
 	}
 
 	// ------------------------------------------------
