Adding new RHI unit test for bindless uniform buffers.

#Synced-CL 52455518

diff --git a/Engine/Plugins/Tests/RHITests/Shaders/Private/TestBindlessUniformBuffer.usf b/Engine/Plugins/Tests/RHITests/Shaders/Private/TestBindlessUniformBuffer.usf
new file mode 100644
--- /dev/null
+++ b/Engine/Plugins/Tests/RHITests/Shaders/Private/TestBindlessUniformBuffer.usf
@@ -0,0 +1,19 @@
+// Copyright Epic Games, Inc. All Rights Reserved.
+
+#include "/Engine/Private/Common.ush"
+
+// Uniform buffer declared with BindlessAccessible.
+// When SUPPORTS_BINDLESS_UNIFORM_BUFFER is enabled, the generated header exposes a
+// GetBindlessAccessibleUniformBufferParametersConstants(uint Index) getter instead of a bound cbuffer.
+#include "/Engine/Generated/UniformBuffers/BindlessAccessibleUniformBufferParameters.ush"
+
+uint BindlessUBIndex;
+RWStructuredBuffer<float4> OutputUAV;
+
+[numthreads(1, 1, 1)]
+void ReadBindlessUniformBuffer(uint3 GroupID : SV_GroupID)
+{
+	// Load the UB manually from the bindless heap using the index passed from C++.
+	BindlessAccessibleUniformBufferParameters = GetBindlessAccessibleUniformBufferParametersConstants(BindlessUBIndex);
+	OutputUAV[0] = BindlessAccessibleUniformBufferParameters.Value;
+}
diff --git a/Engine/Plugins/Tests/RHITests/Source/RHITests/Private/RHIBindlessTests.cpp b/Engine/Plugins/Tests/RHITests/Source/RHITests/Private/RHIBindlessTests.cpp
--- a/Engine/Plugins/Tests/RHITests/Source/RHITests/Private/RHIBindlessTests.cpp
+++ b/Engine/Plugins/Tests/RHITests/Source/RHITests/Private/RHIBindlessTests.cpp
@@ -701,3 +701,100 @@ bool RHIBindlessTests::Test_DescriptorRange_MixedResourceTypes(FRHICommandListIm
 	return ReadbackAndVerifyColors(RHICmdList, OutputBuffer, OutputBufferDesc.Size, NumResources, ExpectedColors, TEXT("Descriptor Range Mixed Resource Types"));
 }
 
+// Uniform buffer struct used to test BindlessAccessible uniform buffers.
+BEGIN_UNIFORM_BUFFER_STRUCT(FBindlessAccessibleUniformBufferParameters, )
+	SHADER_PARAMETER(FVector4f, Value)
+END_UNIFORM_BUFFER_STRUCT()
+IMPLEMENT_UNIFORM_BUFFER_STRUCT_EX(FBindlessAccessibleUniformBufferParameters, "BindlessAccessibleUniformBufferParameters", FShaderParametersMetadata::EUsageFlags::BindlessAccessible);
+
+class FBindlessUniformBufferTestCS : public FGlobalShader
+{
+public:
+	DECLARE_GLOBAL_SHADER(FBindlessUniformBufferTestCS);
+	SHADER_USE_PARAMETER_STRUCT(FBindlessUniformBufferTestCS, FGlobalShader);
+
+	BEGIN_SHADER_PARAMETER_STRUCT(FParameters, )
+		SHADER_PARAMETER_STRUCT_REF(FBindlessAccessibleUniformBufferParameters, BindlessUB)
+		SHADER_PARAMETER(uint32, BindlessUBIndex)
+		SHADER_PARAMETER_UAV(RWStructuredBuffer<float4>, OutputUAV)
+	END_SHADER_PARAMETER_STRUCT()
+
+	static bool ShouldCompilePermutation(const FShaderPermutationParameters& Parameters)
+	{
+		return IsBindlessEnabledForAnyGraphics(FShaderPlatformConfig::GetBindlessConfiguration(Parameters.Platform));
+	}
+
+	static void ModifyCompilationEnvironment(const FGlobalShaderPermutationParameters& Parameters, FShaderCompilerEnvironment& OutEnvironment)
+	{
+		FGlobalShader::ModifyCompilationEnvironment(Parameters, OutEnvironment);
+		OutEnvironment.CompilerFlags.Add(CFLAG_SupportsMinimalBindless);
+		OutEnvironment.SetDefine(TEXT("SUPPORTS_BINDLESS_UNIFORM_BUFFER"), 1);
+	}
+};
+IMPLEMENT_GLOBAL_SHADER(FBindlessUniformBufferTestCS, "/Plugin/RHITests/Private/TestBindlessUniformBuffer.usf", "ReadBindlessUniformBuffer", SF_Compute);
+
+#include "RenderCaptureInterface.h"
+
+// Test that a uniform buffer declared with BindlessAccessible gets a valid bindless handle,
+// and that its contents are correctly accessible when bound to a shader.
+bool RHIBindlessTests::Test_BindlessAccessibleUniformBuffer(FRHICommandListImmediate& RHICmdList)
+{
+	if (!GRHIGlobals.bSupportsBindless || !IsBindlessEnabledForAnyGraphics(FShaderPlatformConfig::GetBindlessConfiguration(GMaxRHIShaderPlatform)))
+	{
+		return true;
+	}
+
+	RenderCaptureInterface::FScopedCapture ScopedCapture(true, &RHICmdList, TEXT("Test_BindlessAccessibleUniformBuffer"));
+
+	const FVector4f TestValue(1.0f, 2.0f, 3.0f, 4.0f);
+
+	FBindlessAccessibleUniformBufferParameters UBContents;
+	UBContents.Value = TestValue;
+
+	TUniformBufferRef<FBindlessAccessibleUniformBufferParameters> UniformBuffer =
+		CreateUniformBufferImmediate(UBContents, EUniformBufferUsage::UniformBuffer_SingleFrame);
+
+	if (!UniformBuffer)
+	{
+		UE_LOG(LogRHIUnitTestCommandlet, Error, TEXT("Test failed. \"Bindless Accessible Uniform Buffer\": Failed to create uniform buffer."));
+		return false;
+	}
+
+	// Verify the bindless handle was allocated.
+	const FRHIDescriptorHandle BindlessHandle = UniformBuffer->GetBindlessHandle();
+	if (!BindlessHandle.IsValid())
+	{
+		UE_LOG(LogRHIUnitTestCommandlet, Error, TEXT("Test failed. \"Bindless Accessible Uniform Buffer\": Uniform buffer does not have a valid bindless handle."));
+		return false;
+	}
+
+	// Dispatch a shader that reads the UB value and writes it to an output buffer.
+	const FRHIBufferCreateDesc OutputBufferDesc =
+		FRHIBufferCreateDesc::CreateStructured<FLinearColor>(TEXT("BindlessUBTestOutput"), 1)
+		.AddUsage(EBufferUsageFlags::UnorderedAccess | EBufferUsageFlags::SourceCopy)
+		.SetInitialState(ERHIAccess::UAVCompute);
+	FBufferRHIRef OutputBuffer = RHICmdList.CreateBuffer(OutputBufferDesc);
+
+	FUnorderedAccessViewRHIRef OutputUAV = RHICmdList.CreateUnorderedAccessView(
+		OutputBuffer,
+		FRHIViewDesc::CreateBufferUAV()
+			.SetType(FRHIViewDesc::EBufferType::Structured)
+			.SetStride(sizeof(FLinearColor)).SetNumElements(1)
+	);
+
+	{
+		TShaderMapRef<FBindlessUniformBufferTestCS> ComputeShader(GetGlobalShaderMap(GMaxRHIFeatureLevel));
+
+		FBindlessUniformBufferTestCS::FParameters Parameters;
+		Parameters.BindlessUB = UniformBuffer;
+		Parameters.BindlessUBIndex = UniformBuffer->GetBindlessHandle().GetIndex();
+		Parameters.OutputUAV = OutputUAV;
+
+		FComputeShaderUtils::Dispatch(RHICmdList, ComputeShader, Parameters, FIntVector(1, 1, 1));
+
+		RHICmdList.Transition(FRHITransitionInfo(OutputUAV, ERHIAccess::UAVCompute, ERHIAccess::CopySrc));
+	}
+
+	const TArray<FLinearColor> ExpectedColors = { FLinearColor(TestValue.X, TestValue.Y, TestValue.Z, TestValue.W) };
+	return ReadbackAndVerifyColors(RHICmdList, OutputBuffer, OutputBufferDesc.Size, 1, ExpectedColors, TEXT("Bindless Accessible Uniform Buffer"));
+}
diff --git a/Engine/Plugins/Tests/RHITests/Source/RHITests/Private/RHIBindlessTests.h b/Engine/Plugins/Tests/RHITests/Source/RHITests/Private/RHIBindlessTests.h
--- a/Engine/Plugins/Tests/RHITests/Source/RHITests/Private/RHIBindlessTests.h
+++ b/Engine/Plugins/Tests/RHITests/Source/RHITests/Private/RHIBindlessTests.h
@@ -12,4 +12,5 @@ namespace RHIBindlessTests
 	bool Test_DescriptorRange_UAV(FRHICommandListImmediate& RHICmdList);
 	bool Test_DescriptorRange_SRV_And_UAV(FRHICommandListImmediate& RHICmdList);
 	bool Test_DescriptorRange_MixedResourceTypes(FRHICommandListImmediate& RHICmdList);
+	bool Test_BindlessAccessibleUniformBuffer(FRHICommandListImmediate& RHICmdList);
 }
diff --git a/Engine/Plugins/Tests/RHITests/Source/RHITests/Private/RHITests.spec.cpp b/Engine/Plugins/Tests/RHITests/Source/RHITests/Private/RHITests.spec.cpp
--- a/Engine/Plugins/Tests/RHITests/Source/RHITests/Private/RHITests.spec.cpp
+++ b/Engine/Plugins/Tests/RHITests/Source/RHITests/Private/RHITests.spec.cpp
@@ -244,5 +244,10 @@ void FAutomationRHITest::Define()
 			bool bResult = RunOnRenderThreadSynchronous(RHIBindlessTests::Test_DescriptorRange_MixedResourceTypes);
 			TestEqual("Descriptor Range Mixed Resource Types failed", bResult, 1);
 		});
+		It("RHI Bindless Accessible Uniform Buffer", [this]()
+		{
+			bool bResult = RunOnRenderThreadSynchronous(RHIBindlessTests::Test_BindlessAccessibleUniformBuffer);
+			TestEqual("Bindless Accessible Uniform Buffer failed", bResult, 1);
+		});
 	});
 }
diff --git a/Engine/Plugins/Tests/RHITests/Source/RHITests/Private/RHITestsModule.cpp b/Engine/Plugins/Tests/RHITests/Source/RHITests/Private/RHITestsModule.cpp
--- a/Engine/Plugins/Tests/RHITests/Source/RHITests/Private/RHITestsModule.cpp
+++ b/Engine/Plugins/Tests/RHITests/Source/RHITests/Private/RHITestsModule.cpp
@@ -131,6 +131,7 @@ static bool RunTests_RenderThread(FRHICommandListImmediate& RHICmdList)
 		RUN_TEST(RHIBindlessTests::Test_DescriptorRange_UAV(RHICmdList));
 		RUN_TEST(RHIBindlessTests::Test_DescriptorRange_SRV_And_UAV(RHICmdList));
 		RUN_TEST(RHIBindlessTests::Test_DescriptorRange_MixedResourceTypes(RHICmdList));
+		RUN_TEST(RHIBindlessTests::Test_BindlessAccessibleUniformBuffer(RHICmdList));
 	}
 
 	// ------------------------------------------------
