Disabling new test until platform errors are addressed.

#Synced-CL 52517224

diff --git a/Engine/Plugins/Tests/RHITests/Source/RHITests/Private/RHIBindlessTests.cpp b/Engine/Plugins/Tests/RHITests/Source/RHITests/Private/RHIBindlessTests.cpp
--- a/Engine/Plugins/Tests/RHITests/Source/RHITests/Private/RHIBindlessTests.cpp
+++ b/Engine/Plugins/Tests/RHITests/Source/RHITests/Private/RHIBindlessTests.cpp
@@ -707,6 +707,8 @@ BEGIN_UNIFORM_BUFFER_STRUCT(FBindlessAccessibleUniformBufferParameters, )
 END_UNIFORM_BUFFER_STRUCT()
 IMPLEMENT_UNIFORM_BUFFER_STRUCT_EX(FBindlessAccessibleUniformBufferParameters, "BindlessAccessibleUniformBufferParameters", FShaderParametersMetadata::EUsageFlags::BindlessAccessible);
 
+static bool GEnableBindlessUniformBufferTest = false;
+
 class FBindlessUniformBufferTestCS : public FGlobalShader
 {
 public:
@@ -721,7 +723,7 @@ public:
 
 	static bool ShouldCompilePermutation(const FShaderPermutationParameters& Parameters)
 	{
-		return IsBindlessEnabledForAnyGraphics(FShaderPlatformConfig::GetBindlessConfiguration(Parameters.Platform));
+		return GEnableBindlessUniformBufferTest && IsBindlessEnabledForAnyGraphics(FShaderPlatformConfig::GetBindlessConfiguration(Parameters.Platform));
 	}
 
 	static void ModifyCompilationEnvironment(const FGlobalShaderPermutationParameters& Parameters, FShaderCompilerEnvironment& OutEnvironment)
@@ -733,19 +735,15 @@ public:
 };
 IMPLEMENT_GLOBAL_SHADER(FBindlessUniformBufferTestCS, "/Plugin/RHITests/Private/TestBindlessUniformBuffer.usf", "ReadBindlessUniformBuffer", SF_Compute);
 
-#include "RenderCaptureInterface.h"
-
 // Test that a uniform buffer declared with BindlessAccessible gets a valid bindless handle,
 // and that its contents are correctly accessible when bound to a shader.
 bool RHIBindlessTests::Test_BindlessAccessibleUniformBuffer(FRHICommandListImmediate& RHICmdList)
 {
-	if (!GRHIGlobals.bSupportsBindless || !IsBindlessEnabledForAnyGraphics(FShaderPlatformConfig::GetBindlessConfiguration(GMaxRHIShaderPlatform)))
+	if (!GEnableBindlessUniformBufferTest || !GRHIGlobals.bSupportsBindless || !IsBindlessEnabledForAnyGraphics(FShaderPlatformConfig::GetBindlessConfiguration(GMaxRHIShaderPlatform)))
 	{
 		return true;
 	}
 
-	RenderCaptureInterface::FScopedCapture ScopedCapture(true, &RHICmdList, TEXT("Test_BindlessAccessibleUniformBuffer"));
-
 	const FVector4f TestValue(1.0f, 2.0f, 3.0f, 4.0f);
 
 	FBindlessAccessibleUniformBufferParameters UBContents;
