Adding unit tests for RHI Root Constants.
- First test uses root constants on a Graphics draw.
- Second test uses root constants on a Compute dispatch.
- Both tests use all 4 components of the root constants.

#Synced-CL 52684686

diff --git a/Engine/Plugins/Tests/RHITests/Shaders/Private/TestRootConstants.usf b/Engine/Plugins/Tests/RHITests/Shaders/Private/TestRootConstants.usf
new file mode 100644
--- /dev/null
+++ b/Engine/Plugins/Tests/RHITests/Shaders/Private/TestRootConstants.usf
@@ -0,0 +1,35 @@
+// Copyright Epic Games, Inc. All Rights Reserved.
+
+#include "/Engine/Private/Common.ush"
+#include "/Engine/Public/RootConstants.ush"
+
+// ---- Compute shader ----
+
+RWStructuredBuffer<uint4> OutputUAV;
+
+[numthreads(1, 1, 1)]
+void WriteRootConstants(uint3 GroupID : SV_GroupID)
+{
+	OutputUAV[0] = GetAllRootConstants();
+}
+
+// ---- Graphics shaders ----
+
+// Trivial fullscreen VS — emits a full-screen triangle from vertex ID.
+void RootConstantsMainVS(
+	in  uint  VertexID  : SV_VertexID,
+	out float4 OutPos   : SV_Position)
+{
+	// Two triangles covering the screen from three vertices
+	float2 UV = float2((VertexID << 1) & 2, VertexID & 2);
+	OutPos = float4(UV * float2(2, -2) + float2(-1, 1), 0, 1);
+}
+
+// PS — writes the root constant values into a UAV so they can be read back.
+RWStructuredBuffer<uint4> GraphicsOutputUAV;
+
+void RootConstantsMainPS(out float4 OutColor : SV_Target0)
+{
+	GraphicsOutputUAV[0] = GetAllRootConstants();
+	OutColor = float4(0, 0, 0, 0);
+}
diff --git a/Engine/Plugins/Tests/RHITests/Source/RHITests/Private/RHIRootConstantsTests.cpp b/Engine/Plugins/Tests/RHITests/Source/RHITests/Private/RHIRootConstantsTests.cpp
new file mode 100644
--- /dev/null
+++ b/Engine/Plugins/Tests/RHITests/Source/RHITests/Private/RHIRootConstantsTests.cpp
@@ -0,0 +1,225 @@
+// Copyright Epic Games, Inc. All Rights Reserved.
+
+#include "RHIRootConstantsTests.h"
+
+#include "RHI.h"
+#include "RHIGPUReadback.h"
+#include "RHIResourceUtils.h"
+#include "RHIStaticStates.h"
+#include "RHITestsCommon.h"
+#include "DataDrivenShaderPlatformInfo.h"
+#include "GlobalShader.h"
+#include "PipelineStateCache.h"
+#include "RenderGraphUtils.h"
+#include "ShaderCompilerCore.h"
+#include "ShaderParameterStruct.h"
+#include "CommonRenderResources.h"
+
+// ---- Shaders ----
+
+class FRootConstantsTestShaderBase : public FGlobalShader
+{
+public:
+	FRootConstantsTestShaderBase() = default;
+	FRootConstantsTestShaderBase(const ShaderMetaType::CompiledShaderInitializerType& Initializer)
+		: FGlobalShader(Initializer)
+	{
+	}
+	
+	static bool ShouldCompilePermutation(const FShaderPermutationParameters& Parameters)
+	{
+		return FDataDrivenShaderPlatformInfo::GetSupportsShaderRootConstants(Parameters.Platform);
+	}
+
+	static void ModifyCompilationEnvironment(const FGlobalShaderPermutationParameters& Parameters, FShaderCompilerEnvironment& OutEnvironment)
+	{
+		FGlobalShader::ModifyCompilationEnvironment(Parameters, OutEnvironment);
+		OutEnvironment.CompilerFlags.Add(CFLAG_RootConstants);
+	}
+};
+
+class FRootConstantsTestVS : public FRootConstantsTestShaderBase
+{
+public:
+	DECLARE_GLOBAL_SHADER(FRootConstantsTestVS);
+	SHADER_USE_PARAMETER_STRUCT(FRootConstantsTestVS, FRootConstantsTestShaderBase);
+
+	BEGIN_SHADER_PARAMETER_STRUCT(FParameters, )
+	END_SHADER_PARAMETER_STRUCT()
+};
+IMPLEMENT_GLOBAL_SHADER(FRootConstantsTestVS, "/Plugin/RHITests/Private/TestRootConstants.usf", "RootConstantsMainVS", SF_Vertex);
+
+class FRootConstantsTestPS : public FRootConstantsTestShaderBase
+{
+public:
+	DECLARE_GLOBAL_SHADER(FRootConstantsTestPS);
+	SHADER_USE_PARAMETER_STRUCT(FRootConstantsTestPS, FRootConstantsTestShaderBase);
+
+	BEGIN_SHADER_PARAMETER_STRUCT(FParameters, )
+		SHADER_PARAMETER_UAV(RWStructuredBuffer<uint4>, GraphicsOutputUAV)
+	END_SHADER_PARAMETER_STRUCT()
+};
+IMPLEMENT_GLOBAL_SHADER(FRootConstantsTestPS, "/Plugin/RHITests/Private/TestRootConstants.usf", "RootConstantsMainPS", SF_Pixel);
+
+class FRootConstantsTestCS : public FRootConstantsTestShaderBase
+{
+public:
+	DECLARE_GLOBAL_SHADER(FRootConstantsTestCS);
+	SHADER_USE_PARAMETER_STRUCT(FRootConstantsTestCS, FRootConstantsTestShaderBase);
+
+	BEGIN_SHADER_PARAMETER_STRUCT(FParameters, )
+		SHADER_PARAMETER_UAV(RWStructuredBuffer<uint4>, OutputUAV)
+	END_SHADER_PARAMETER_STRUCT()
+};
+IMPLEMENT_GLOBAL_SHADER(FRootConstantsTestCS, "/Plugin/RHITests/Private/TestRootConstants.usf", "WriteRootConstants", SF_Compute);
+
+// ---- Helpers ----
+
+static bool ReadbackAndVerifyRootConstants(
+	FRHICommandListImmediate& RHICmdList,
+	FBufferRHIRef OutputBuffer,
+	uint32 BufferSize,
+	const FUint32Vector4& ExpectedValue,
+	const TCHAR* TestName)
+{
+	FRHIGPUBufferReadback Readback(TestName);
+	Readback.EnqueueCopy(RHICmdList, OutputBuffer, BufferSize);
+
+	RHICmdList.BlockUntilGPUIdle();
+	RHICmdList.FlushResources();
+
+	const FUint32Vector4* OutputData = reinterpret_cast<const FUint32Vector4*>(Readback.Lock(BufferSize));
+
+	bool bResult = true;
+	if (*OutputData != ExpectedValue)
+	{
+		UE_LOG(LogRHIUnitTestCommandlet, Error,
+			TEXT("Test failed. \"%s\"\n\nExpected: (%u, %u, %u, %u); Actual: (%u, %u, %u, %u)"),
+			TestName,
+			ExpectedValue.X, ExpectedValue.Y, ExpectedValue.Z, ExpectedValue.W,
+			OutputData->X, OutputData->Y, OutputData->Z, OutputData->W);
+		bResult = false;
+	}
+
+	Readback.Unlock();
+
+	if (bResult)
+	{
+		UE_LOG(LogRHIUnitTestCommandlet, Display, TEXT("Test passed. \"%s\""), TestName);
+	}
+
+	return bResult;
+}
+
+// ---- Tests ----
+
+// Draws a fullscreen triangle with a pixel shader that writes the graphics root constants
+// to an output UAV.
+bool RHIRootConstantsTests::Test_GraphicsShaderRootConstants(FRHICommandListImmediate& RHICmdList)
+{
+	if (!GRHISupportsShaderRootConstants || !GRHIGlobals.SupportsPixelShaderUAVs)
+	{
+		return true;
+	}
+
+	const FUint32Vector4 GraphicsValue(1u, 2u, 3u, 4u);
+
+	// Create output UAV
+	const FRHIBufferCreateDesc OutputBufferDesc =
+		FRHIBufferCreateDesc::CreateStructured<FUint32Vector4>(TEXT("GraphicsRootConstantsOutput"), 1)
+		.AddUsage(EBufferUsageFlags::UnorderedAccess | EBufferUsageFlags::SourceCopy)
+		.SetInitialState(ERHIAccess::UAVCompute);
+	FBufferRHIRef OutputBuffer = RHICmdList.CreateBuffer(OutputBufferDesc);
+
+	FUnorderedAccessViewRHIRef OutputUAV = RHICmdList.CreateUnorderedAccessView(
+		OutputBuffer,
+		FRHIViewDesc::CreateBufferUAV()
+			.SetType(FRHIViewDesc::EBufferType::Structured)
+			.SetStride(sizeof(FUint32Vector4)).SetNumElements(1)
+	);
+
+	// Create a 1x1 render target (we only need the pixel shader UAV write)
+	const FRHITextureCreateDesc RTDesc =
+		FRHITextureCreateDesc::Create2D(TEXT("RootConstantsTest_RT"), 1, 1, PF_B8G8R8A8)
+		.SetFlags(ETextureCreateFlags::RenderTargetable)
+		.SetInitialState(ERHIAccess::RTV);
+	FTextureRHIRef RenderTarget = RHICmdList.CreateTexture(RTDesc);
+
+	TShaderMapRef<FRootConstantsTestVS> VertexShader(GetGlobalShaderMap(GMaxRHIFeatureLevel));
+	TShaderMapRef<FRootConstantsTestPS> PixelShader(GetGlobalShaderMap(GMaxRHIFeatureLevel));
+
+	FGraphicsPipelineStateInitializer GraphicsPSOInit;
+	GraphicsPSOInit.BoundShaderState.VertexShaderRHI = VertexShader.GetVertexShader();
+	GraphicsPSOInit.BoundShaderState.PixelShaderRHI  = PixelShader.GetPixelShader();
+	GraphicsPSOInit.BoundShaderState.VertexDeclarationRHI = GEmptyVertexDeclaration.VertexDeclarationRHI;
+	GraphicsPSOInit.DepthStencilState  = TStaticDepthStencilState<false, CF_Always>::GetRHI();
+	GraphicsPSOInit.BlendState         = TStaticBlendState<>::GetRHI();
+	GraphicsPSOInit.RasterizerState    = TStaticRasterizerState<>::GetRHI();
+	GraphicsPSOInit.PrimitiveType      = EPrimitiveType::PT_TriangleList;
+
+	RHICmdList.Transition(FRHITransitionInfo(OutputUAV, ERHIAccess::UAVCompute, ERHIAccess::UAVGraphics));
+
+	FRHITexture* ColorRTs[1] = { RenderTarget.GetReference() };
+	FRHIRenderPassInfo RenderPassInfo(1, ColorRTs, ERenderTargetActions::DontLoad_DontStore);
+	RHICmdList.BeginRenderPass(RenderPassInfo, TEXT("RootConstantsGraphicsTest"));
+	{
+		RHICmdList.SetViewport(0, 0, 0, 1, 1, 1);
+		RHICmdList.ApplyCachedRenderTargets(GraphicsPSOInit);
+		SetGraphicsPipelineState(RHICmdList, GraphicsPSOInit, 0);
+
+		RHICmdList.SetShaderRootConstants(GraphicsValue);
+
+		FRootConstantsTestPS::FParameters PSParams;
+		PSParams.GraphicsOutputUAV = OutputUAV;
+
+		FRHIBatchedShaderParameters& ShaderParams = RHICmdList.GetScratchShaderParameters();
+		SetShaderParameters(ShaderParams, PixelShader, PSParams);
+		RHICmdList.SetBatchedShaderParameters(PixelShader.GetPixelShader(), ShaderParams);
+
+		RHICmdList.DrawPrimitive(0, 1, 1);
+	}
+	RHICmdList.EndRenderPass();
+
+	RHICmdList.Transition(FRHITransitionInfo(OutputUAV, ERHIAccess::UAVGraphics, ERHIAccess::CopySrc));
+
+	return ReadbackAndVerifyRootConstants(RHICmdList, OutputBuffer, OutputBufferDesc.Size, GraphicsValue, TEXT("Graphics Shader Root Constants"));
+}
+
+// Dispatches a compute shader that writes root constants to a structured buffer.
+bool RHIRootConstantsTests::Test_ComputeShaderRootConstants(FRHICommandListImmediate& RHICmdList)
+{
+	if (!GRHISupportsShaderRootConstants)
+	{
+		return true;
+	}
+
+	const FUint32Vector4 ComputeValue(11u, 22u, 33u, 44u);
+
+	const FRHIBufferCreateDesc OutputBufferDesc =
+		FRHIBufferCreateDesc::CreateStructured<FUint32Vector4>(TEXT("ComputeRootConstantsOutput"), 1)
+		.AddUsage(EBufferUsageFlags::UnorderedAccess | EBufferUsageFlags::SourceCopy)
+		.SetInitialState(ERHIAccess::UAVCompute);
+	FBufferRHIRef OutputBuffer = RHICmdList.CreateBuffer(OutputBufferDesc);
+
+	FUnorderedAccessViewRHIRef OutputUAV = RHICmdList.CreateUnorderedAccessView(
+		OutputBuffer,
+		FRHIViewDesc::CreateBufferUAV()
+			.SetType(FRHIViewDesc::EBufferType::Structured)
+			.SetStride(sizeof(FUint32Vector4)).SetNumElements(1)
+	);
+
+	{
+		TShaderMapRef<FRootConstantsTestCS> ComputeShader(GetGlobalShaderMap(GMaxRHIFeatureLevel));
+
+		RHICmdList.SetComputeShaderRootConstants(ComputeValue);
+
+		FRootConstantsTestCS::FParameters Parameters;
+		Parameters.OutputUAV = OutputUAV;
+
+		FComputeShaderUtils::Dispatch(RHICmdList, ComputeShader, Parameters, FIntVector(1, 1, 1));
+
+		RHICmdList.Transition(FRHITransitionInfo(OutputUAV, ERHIAccess::UAVCompute, ERHIAccess::CopySrc));
+	}
+
+	return ReadbackAndVerifyRootConstants(RHICmdList, OutputBuffer, OutputBufferDesc.Size, ComputeValue, TEXT("Compute Shader Root Constants"));
+}
diff --git a/Engine/Plugins/Tests/RHITests/Source/RHITests/Private/RHIRootConstantsTests.h b/Engine/Plugins/Tests/RHITests/Source/RHITests/Private/RHIRootConstantsTests.h
new file mode 100644
--- /dev/null
+++ b/Engine/Plugins/Tests/RHITests/Source/RHITests/Private/RHIRootConstantsTests.h
@@ -0,0 +1,11 @@
+// Copyright Epic Games, Inc. All Rights Reserved.
+
+#pragma once
+
+class FRHICommandListImmediate;
+
+namespace RHIRootConstantsTests
+{
+	bool Test_GraphicsShaderRootConstants(FRHICommandListImmediate& RHICmdList);
+	bool Test_ComputeShaderRootConstants(FRHICommandListImmediate& RHICmdList);
+}
diff --git a/Engine/Plugins/Tests/RHITests/Source/RHITests/Private/RHITests.spec.cpp b/Engine/Plugins/Tests/RHITests/Source/RHITests/Private/RHITests.spec.cpp
--- a/Engine/Plugins/Tests/RHITests/Source/RHITests/Private/RHITests.spec.cpp
+++ b/Engine/Plugins/Tests/RHITests/Source/RHITests/Private/RHITests.spec.cpp
@@ -13,6 +13,7 @@
 #include "RHIGraphicsUAVTests.h"
 #include "RHIAllocatorTests.h"
 #include "RHIBindlessTests.h"
+#include "RHIRootConstantsTests.h"
 
 BEGIN_DEFINE_SPEC(FAutomationRHITest, "Rendering.RHI", EAutomationTestFlags::EngineFilter | EAutomationTestFlags_ApplicationContextMask | EAutomationTestFlags::NonNullRHI)
 END_DEFINE_SPEC(FAutomationRHITest)
@@ -249,5 +250,21 @@ void FAutomationRHITest::Define()
 			bool bResult = RunOnRenderThreadSynchronous(RHIBindlessTests::Test_BindlessAccessibleUniformBuffer);
 			TestEqual("Bindless Accessible Uniform Buffer failed", bResult, 1);
 		});
+
+	});
+
+	Describe("Root Constants", [this]
+	{
+		It("RHI Graphics Shader Root Constants", [this]()
+		{
+			bool bResult = RunOnRenderThreadSynchronous(RHIRootConstantsTests::Test_GraphicsShaderRootConstants);
+			TestEqual("Graphics Shader Root Constants failed", bResult, 1);
+		});
+
+		It("RHI Compute Shader Root Constants", [this]()
+		{
+			bool bResult = RunOnRenderThreadSynchronous(RHIRootConstantsTests::Test_ComputeShaderRootConstants);
+			TestEqual("Compute Shader Root Constants failed", bResult, 1);
+		});
 	});
 }
diff --git a/Engine/Plugins/Tests/RHITests/Source/RHITests/Private/RHITestsModule.cpp b/Engine/Plugins/Tests/RHITests/Source/RHITests/Private/RHITestsModule.cpp
--- a/Engine/Plugins/Tests/RHITests/Source/RHITests/Private/RHITestsModule.cpp
+++ b/Engine/Plugins/Tests/RHITests/Source/RHITests/Private/RHITestsModule.cpp
@@ -14,6 +14,7 @@
 #include "RHIGraphicsUAVTests.h"
 #include "RHIAllocatorTests.h"
 #include "RHIBindlessTests.h"
+#include "RHIRootConstantsTests.h"
 #include "RHIShaderBundleTests.h"
 
 #define LOCTEXT_NAMESPACE "FRHITestsModule"
@@ -134,6 +135,14 @@ static bool RunTests_RenderThread(FRHICommandListImmediate& RHICmdList)
 		RUN_TEST(RHIBindlessTests::Test_BindlessAccessibleUniformBuffer(RHICmdList));
 	}
 
+	// ------------------------------------------------
+	// Root Constants
+	// ------------------------------------------------
+	{
+		RUN_TEST(RHIRootConstantsTests::Test_GraphicsShaderRootConstants(RHICmdList));
+		RUN_TEST(RHIRootConstantsTests::Test_ComputeShaderRootConstants(RHICmdList));
+	}
+
 	// ------------------------------------------------
 	// Shader Bundles
 	// ------------------------------------------------
