- Removed BlockUntilGPUIdle and SubmitCommandsAndFlushGPU in place of SubmitAndBlockUntilGPUIdle
- Removed FRHICommandList::FlushResources and RHIFlushResources
- Moved Vulkan Command Buffer cleanup
- SubmitAndBlockUntilGPUIdle now takes a parameter to flush resources

#Synced-CL 52187599

diff --git a/Engine/Platforms/VisionOS/Plugins/Runtime/OpenXRVisionOS/Source/OXRVisionOS/Private/OXRVisionOSSession.cpp b/Engine/Platforms/VisionOS/Plugins/Runtime/OpenXRVisionOS/Source/OXRVisionOS/Private/OXRVisionOSSession.cpp
--- a/Engine/Platforms/VisionOS/Plugins/Runtime/OpenXRVisionOS/Source/OXRVisionOS/Private/OXRVisionOSSession.cpp
+++ b/Engine/Platforms/VisionOS/Plugins/Runtime/OpenXRVisionOS/Source/OXRVisionOS/Private/OXRVisionOSSession.cpp
@@ -201,7 +201,7 @@ FOXRVisionOSSession::~FOXRVisionOSSession()
 				// Perform a full flush
 				RHICmdList.ImmediateFlush(EImmediateFlushType::FlushRHIThreadFlushResources);
 
-				RHICmdList.BlockUntilGPUIdle();
+				RHICmdList.SubmitAndBlockUntilGPUIdle();
 			});
 
 		FlushRenderingCommands();
diff --git a/Engine/Plugins/Compositing/Composite/Source/Composite/Private/Tests/CompositeTestHelpers.h b/Engine/Plugins/Compositing/Composite/Source/Composite/Private/Tests/CompositeTestHelpers.h
--- a/Engine/Plugins/Compositing/Composite/Source/Composite/Private/Tests/CompositeTestHelpers.h
+++ b/Engine/Plugins/Compositing/Composite/Source/Composite/Private/Tests/CompositeTestHelpers.h
@@ -54,7 +54,7 @@ public:
 
 // Enqueues a batch of RDG passes on the render thread and blocks the game thread until
 // they complete. Work receives an open FRDGBuilder and the view; GraphBuilder.Execute()
-// and BlockUntilGPUIdle() are called internally so callers only need to add passes.
+// and SubmitAndBlockUntilGPUIdle() are called internally so callers only need to add passes.
 // PostWork (optional) runs after Execute() and BlockUntilGPUIdle() — use it to lock
 // GPU readbacks (FRHIGPUTextureReadback::Lock) that were enqueued during Work.
 inline void DispatchAndWait(
@@ -70,7 +70,7 @@ inline void DispatchAndWait(
 			FRDGBuilder GraphBuilder(RHICmdList);
 			Work(GraphBuilder, *ViewPtr);
 			GraphBuilder.Execute();
-			RHICmdList.BlockUntilGPUIdle();
+			RHICmdList.SubmitAndBlockUntilGPUIdle();
 			if (PostWork) { PostWork(RHICmdList); }
 			Signal->Trigger();
 		});
diff --git a/Engine/Plugins/Developer/PixWinPlugin/Source/PixWinPlugin/Private/PixWinPluginModule.cpp b/Engine/Plugins/Developer/PixWinPlugin/Source/PixWinPlugin/Private/PixWinPluginModule.cpp
--- a/Engine/Plugins/Developer/PixWinPlugin/Source/PixWinPlugin/Private/PixWinPluginModule.cpp
+++ b/Engine/Plugins/Developer/PixWinPlugin/Source/PixWinPlugin/Private/PixWinPluginModule.cpp
@@ -222,7 +222,7 @@ void FPixWinPluginModule::BeginCapture(FRHICommandListImmediate* InRHICommandLis
 #if PIX_PLUGIN_ENABLED
 	CurrentCaptureDestFileName = MakeWinPixCaptureFilePath(InDestFileName);
 	CurrentCaptureFlags = InFlags;
-	InRHICommandList->SubmitCommandsAndFlushGPU();
+	InRHICommandList->SubmitAndBlockUntilGPUIdle();
 	InRHICommandList->EnqueueLambda([Pix = PixGraphicsAnalysisInterface, DestFileName = CurrentCaptureDestFileName](FRHICommandListImmediate& RHICommandList)
 	{
 		Pix->BeginCapture(nullptr, DestFileName);
@@ -233,7 +233,7 @@ void FPixWinPluginModule::BeginCapture(FRHICommandListImmediate* InRHICommandLis
 void FPixWinPluginModule::EndCapture(FRHICommandListImmediate* InRHICommandList)
 {
 #if PIX_PLUGIN_ENABLED
-	InRHICommandList->SubmitCommandsAndFlushGPU();
+	InRHICommandList->SubmitAndBlockUntilGPUIdle();
 	InRHICommandList->EnqueueLambda(
 		[Pix = PixGraphicsAnalysisInterface, 
 		Flags = CurrentCaptureFlags, 
@@ -369,7 +369,7 @@ void FPixWinPluginModule::EndFrameCapture(uint32 InFlags, const FString& DestFil
 		{
 			if (Pix && Pix->IsValid())
 			{
-				RHICommandList.SubmitCommandsAndFlushGPU();
+				RHICommandList.SubmitAndBlockUntilGPUIdle();
 				Pix->EndCapture();
 
 				// If we're already attached, don't open a new PIX instance
diff --git a/Engine/Plugins/Developer/RenderDocPlugin/Source/RenderDocPlugin/Private/RenderDocPluginModule.cpp b/Engine/Plugins/Developer/RenderDocPlugin/Source/RenderDocPlugin/Private/RenderDocPluginModule.cpp
--- a/Engine/Plugins/Developer/RenderDocPlugin/Source/RenderDocPlugin/Private/RenderDocPluginModule.cpp
+++ b/Engine/Plugins/Developer/RenderDocPlugin/Source/RenderDocPlugin/Private/RenderDocPluginModule.cpp
@@ -132,7 +132,7 @@ public:
 
 	static uint32 EndFrameCapture(FRHICommandListImmediate& RHICmdList, HWND WindowHandle, FRenderDocPluginLoader::RENDERDOC_API_CONTEXT* RenderDocAPI)
 	{
-		RHICmdList.SubmitCommandsAndFlushGPU();
+		RHICmdList.SubmitAndBlockUntilGPUIdle();
 
 		uint32 Result = RenderDocAPI->EndFrameCapture(GetRenderdocDevicePointer(), WindowHandle);
 		SetEmitDrawEvents(UE_GEmitDrawEvents_BeforeCapture);
@@ -701,7 +701,7 @@ void FRenderDocPluginModule::EndCapture(FRHICommandListImmediate* InRHICommandLi
 void FRenderDocPluginModule::BeginCapture_RenderThread(FRHICommandListImmediate* InRHICommandList)
 {
 	RENDERDOC_DevicePointer Device = FRenderDocFrameCapturer::GetRenderdocDevicePointer();
-	InRHICommandList->SubmitCommandsAndFlushGPU();
+	InRHICommandList->SubmitAndBlockUntilGPUIdle();
 	InRHICommandList->EnqueueLambda([this, Device](FRHICommandListImmediate& RHICommandList)
 	{
 		RenderDocAPI->StartFrameCapture(Device, MainWindowHandle);
@@ -711,7 +711,7 @@ void FRenderDocPluginModule::BeginCapture_RenderThread(FRHICommandListImmediate*
 void FRenderDocPluginModule::EndCapture_RenderThread(FRHICommandListImmediate* InRHICommandList, uint32 InFlags, FString const& InDestFileName)
 {
 	RENDERDOC_DevicePointer Device = FRenderDocFrameCapturer::GetRenderdocDevicePointer();
-	InRHICommandList->SubmitCommandsAndFlushGPU();
+	InRHICommandList->SubmitAndBlockUntilGPUIdle();
 	InRHICommandList->EnqueueLambda([this, Device, InFlags, DestFileName = InDestFileName](FRHICommandListImmediate& RHICommandList)
 	{
 		uint32 Result = RenderDocAPI->EndFrameCapture(Device, MainWindowHandle);
diff --git a/Engine/Plugins/Developer/XcodeGPUDebuggerPlugin/Source/XcodeGPUDebuggerPlugin/Private/XcodeGPUDebuggerPluginModule.cpp b/Engine/Plugins/Developer/XcodeGPUDebuggerPlugin/Source/XcodeGPUDebuggerPlugin/Private/XcodeGPUDebuggerPluginModule.cpp
--- a/Engine/Plugins/Developer/XcodeGPUDebuggerPlugin/Source/XcodeGPUDebuggerPlugin/Private/XcodeGPUDebuggerPluginModule.cpp
+++ b/Engine/Plugins/Developer/XcodeGPUDebuggerPlugin/Source/XcodeGPUDebuggerPlugin/Private/XcodeGPUDebuggerPluginModule.cpp
@@ -93,7 +93,7 @@ public:
 
 	static void EndFrameCapture(FRHICommandListImmediate& RHICmdList)
 	{
-		RHICmdList.SubmitCommandsAndFlushGPU();
+		RHICmdList.SubmitAndBlockUntilGPUIdle();
 
 		MTL::CaptureManager* CaptureManager = MTL::CaptureManager::sharedCaptureManager();
 		CaptureManager->stopCapture();
@@ -353,7 +353,7 @@ void FXcodeGPUDebuggerPluginModule::BeginCapture(FRHICommandListImmediate* InRHI
 	
 	MTL::Device* MetalDevice = GetIMetalDynamicRHI()->RHIGetDevice();
 
-	InRHICommandList->SubmitCommandsAndFlushGPU();
+	InRHICommandList->SubmitAndBlockUntilGPUIdle();
 	InRHICommandList->EnqueueLambda([FileName = CaptureFileName, MetalDevice](FRHICommandListImmediate& RHICommandList)
 	{
 		MTL::CaptureManager* CaptureManager = MTL::CaptureManager::sharedCaptureManager();
@@ -384,7 +384,7 @@ void FXcodeGPUDebuggerPluginModule::BeginCapture(FRHICommandListImmediate* InRHI
 
 void FXcodeGPUDebuggerPluginModule::EndCapture(FRHICommandListImmediate* InRHICommandList)
 {
-	InRHICommandList->SubmitCommandsAndFlushGPU();
+	InRHICommandList->SubmitAndBlockUntilGPUIdle();
     InRHICommandList->EnqueueLambda([this](FRHICommandListImmediate& RHICommandList)
 	{
 		MTL::CaptureManager* CaptureManager = MTL::CaptureManager::sharedCaptureManager();
diff --git a/Engine/Plugins/Experimental/GPULightmass/Source/GPULightmass/Private/Scene/Lights.cpp b/Engine/Plugins/Experimental/GPULightmass/Source/GPULightmass/Private/Scene/Lights.cpp
--- a/Engine/Plugins/Experimental/GPULightmass/Source/GPULightmass/Private/Scene/Lights.cpp
+++ b/Engine/Plugins/Experimental/GPULightmass/Source/GPULightmass/Private/Scene/Lights.cpp
@@ -468,7 +468,7 @@ void FDirectionalLightRenderState::RenderStaticShadowDepthMap(FRHICommandListImm
 	
 	GraphBuilder.Execute();
 
-	RHICmdList.BlockUntilGPUIdle();
+	RHICmdList.SubmitAndBlockUntilGPUIdle();
 
 	check(DepthMapTextureReadback.IsReady());
 
@@ -580,7 +580,7 @@ void FSpotLightRenderState::RenderStaticShadowDepthMap(FRHICommandListImmediate&
 	
 	GraphBuilder.Execute();
 
-	RHICmdList.BlockUntilGPUIdle();
+	RHICmdList.SubmitAndBlockUntilGPUIdle();
 
 	check(DepthMapTextureReadback.IsReady());
 
@@ -683,7 +683,7 @@ void FPointLightRenderState::RenderStaticShadowDepthMap(FRHICommandListImmediate
 	
 	GraphBuilder.Execute();
 
-	RHICmdList.BlockUntilGPUIdle();
+	RHICmdList.SubmitAndBlockUntilGPUIdle();
 
 	check(DepthMapTextureReadback.IsReady());
 
@@ -777,7 +777,7 @@ void FRectLightRenderState::RenderStaticShadowDepthMap(FRHICommandListImmediate&
 	
 	GraphBuilder.Execute();
 
-	RHICmdList.BlockUntilGPUIdle();
+	RHICmdList.SubmitAndBlockUntilGPUIdle();
 
 	check(DepthMapTextureReadback.IsReady());
 
diff --git a/Engine/Plugins/Experimental/GPULightmass/Source/GPULightmass/Private/VolumetricLightmap.cpp b/Engine/Plugins/Experimental/GPULightmass/Source/GPULightmass/Private/VolumetricLightmap.cpp
--- a/Engine/Plugins/Experimental/GPULightmass/Source/GPULightmass/Private/VolumetricLightmap.cpp
+++ b/Engine/Plugins/Experimental/GPULightmass/Source/GPULightmass/Private/VolumetricLightmap.cpp
@@ -420,7 +420,7 @@ void FVolumetricLightmapRenderer::VoxelizeScene(FRHICommandListImmediate& RHICmd
 	{
 		FRHIGPUBufferReadback NumBricksReadback(TEXT("NumBricksReadback"));
 		NumBricksReadback.EnqueueCopy(RHICmdList, BrickAllocatorParameters.Buffer);
-		RHICmdList.BlockUntilGPUIdle();
+		RHICmdList.SubmitAndBlockUntilGPUIdle();
 		check(NumBricksReadback.IsReady());
 
 		int32* Buffer = (int32*)NumBricksReadback.Lock(8);
diff --git a/Engine/Plugins/Experimental/NFORDenoise/Source/NFORDenoise/Private/NFORRegressionCPUSolver.cpp b/Engine/Plugins/Experimental/NFORDenoise/Source/NFORDenoise/Private/NFORRegressionCPUSolver.cpp
--- a/Engine/Plugins/Experimental/NFORDenoise/Source/NFORDenoise/Private/NFORRegressionCPUSolver.cpp
+++ b/Engine/Plugins/Experimental/NFORDenoise/Source/NFORDenoise/Private/NFORRegressionCPUSolver.cpp
@@ -429,7 +429,7 @@ static void CopyTextureFromGPUToCPU(FRHICommandListImmediate& RHICmdList, FRHITe
 {
 	FRHIGPUTextureReadback Readback(TEXT("NFOR::TextureReadback"));
 	Readback.EnqueueCopy(RHICmdList, SrcTexture, FIntVector::ZeroValue, 0, FIntVector(Size.X, Size.Y, 1));
-	RHICmdList.BlockUntilGPUIdle();
+	RHICmdList.SubmitAndBlockUntilGPUIdle();
 
 	int32_t SrcStride = 0;
 	const PixelType* SrcBuffer = static_cast<PixelType*>(Readback.Lock(SrcStride, nullptr));
@@ -463,7 +463,7 @@ static void CopyBufferFromGPUToCPU(FRHICommandListImmediate& RHICmdList, FRHIBuf
 	FRHIGPUBufferReadback Readback(TEXT("NFOR::BufferReadBack"));
 	uint32_t NumBytes = NumElement * sizeof(ElementType);
 	Readback.EnqueueCopy(RHICmdList, SrcBuffer, NumBytes);
-	RHICmdList.BlockUntilGPUIdle();
+	RHICmdList.SubmitAndBlockUntilGPUIdle();
 
 	const ElementType* SrcRawBuffer = static_cast<ElementType*>(Readback.Lock(NumBytes));
 
@@ -578,4 +578,4 @@ void SolveLinearEquationCPU(
 			CopyBufferFromCPUToGPU(RHICmdList, NFORLinearSolveState.B,
 				LinearSolverParameters->B->GetSize() / LinearSolverParameters->B->GetStride(), LinearSolverParameters->X->GetRHI());
 		});
-}
\ No newline at end of file
+}
diff --git a/Engine/Plugins/Experimental/NNERuntimeIREE/Source/IREEDriverRDG/Private/IREEDriverRDGBuiltinExecutablesTest.cpp b/Engine/Plugins/Experimental/NNERuntimeIREE/Source/IREEDriverRDG/Private/IREEDriverRDGBuiltinExecutablesTest.cpp
--- a/Engine/Plugins/Experimental/NNERuntimeIREE/Source/IREEDriverRDG/Private/IREEDriverRDGBuiltinExecutablesTest.cpp
+++ b/Engine/Plugins/Experimental/NNERuntimeIREE/Source/IREEDriverRDG/Private/IREEDriverRDGBuiltinExecutablesTest.cpp
@@ -45,7 +45,7 @@ void CopyBufferFromGPUToCPU(FRHICommandListImmediate& RHICmdList, FRHIBuffer* Sr
 	FRHIGPUBufferReadback Readback(TEXT("IREEDriverRDG.FillBufferTest.CopyBufferFromGPUToCPU"));
 	Readback.EnqueueCopy(RHICmdList, SrcBuffer, NumBytes);
 
-	RHICmdList.BlockUntilGPUIdle();
+	RHICmdList.SubmitAndBlockUntilGPUIdle();
 
 	const ElementType* Src = static_cast<ElementType*>(Readback.Lock(NumBytes));
 	ElementType* Dst = DstArray.GetData();
@@ -113,7 +113,7 @@ bool FFillBufferMainTest::RunTest(const FString& Parameter)
 
 			GraphBuilder.Execute();
 
-			RHICmdList.BlockUntilGPUIdle();
+			RHICmdList.SubmitAndBlockUntilGPUIdle();
 
 			Signal->Trigger();
 		});
@@ -197,4 +197,4 @@ bool FFillBufferExpandPatternTest::RunTest(const FString& Parameter)
 
 } // namespace UE::IREE::HAL::RDG
 
-#endif // WITH_IREE_DRIVER_RDG
\ No newline at end of file
+#endif // WITH_IREE_DRIVER_RDG
diff --git a/Engine/Plugins/Experimental/NNERuntimeIREE/Source/NNERuntimeIREE/Private/NNERuntimeIREEModel.cpp b/Engine/Plugins/Experimental/NNERuntimeIREE/Source/NNERuntimeIREE/Private/NNERuntimeIREEModel.cpp
--- a/Engine/Plugins/Experimental/NNERuntimeIREE/Source/NNERuntimeIREE/Private/NNERuntimeIREEModel.cpp
+++ b/Engine/Plugins/Experimental/NNERuntimeIREE/Source/NNERuntimeIREE/Private/NNERuntimeIREEModel.cpp
@@ -1233,7 +1233,7 @@ namespace UE::NNERuntimeIREE
 
 								GraphBuilder.Execute();
 
-								RHICmdList.BlockUntilGPUIdle();
+								RHICmdList.SubmitAndBlockUntilGPUIdle();
 								Signal->Trigger();
 							}
 						);
@@ -1654,4 +1654,4 @@ namespace UE::NNERuntimeIREE
 #endif // WITH_NNE_RUNTIME_IREE_RDG
 } // UE::NNERuntimeIREE
 
-#endif // WITH_NNE_RUNTIME_IREE
\ No newline at end of file
+#endif // WITH_NNE_RUNTIME_IREE
diff --git a/Engine/Plugins/Experimental/NNERuntimeRDG/Source/NNERuntimeRDG/Private/NNERuntimeRDGModelHlsl.cpp b/Engine/Plugins/Experimental/NNERuntimeRDG/Source/NNERuntimeRDG/Private/NNERuntimeRDGModelHlsl.cpp
--- a/Engine/Plugins/Experimental/NNERuntimeRDG/Source/NNERuntimeRDG/Private/NNERuntimeRDGModelHlsl.cpp
+++ b/Engine/Plugins/Experimental/NNERuntimeRDG/Source/NNERuntimeRDG/Private/NNERuntimeRDGModelHlsl.cpp
@@ -365,7 +365,7 @@ void FModelInstance::PrepareWeights()
 				if (CopyDataFlag == ERDGInitialDataFlags::NoCopy)
 				{
 					//To prevent any problem if model is released before upload is done to the GPU. To be improved.
-					RHICmdList.BlockUntilGPUIdle();
+					RHICmdList.SubmitAndBlockUntilGPUIdle();
 				}
 
 
@@ -399,4 +399,4 @@ TSharedPtr<NNE::IModelInstanceRDG> FModel::CreateModelInstanceRDG()
 
 FModel::FModel(const TSharedPtr<UE::NNE::FSharedModelData>& InModelData) : ModelData(InModelData) {}
 
-} // namespace UE::NNERuntimeRDG::Private::Hlsl
\ No newline at end of file
+} // namespace UE::NNERuntimeRDG::Private::Hlsl
diff --git a/Engine/Plugins/Experimental/NeuralRendering/Source/NeuralPostProcessing/Private/NeuralPostProcessModelInstance.cpp b/Engine/Plugins/Experimental/NeuralRendering/Source/NeuralPostProcessing/Private/NeuralPostProcessModelInstance.cpp
--- a/Engine/Plugins/Experimental/NeuralRendering/Source/NeuralPostProcessing/Private/NeuralPostProcessModelInstance.cpp
+++ b/Engine/Plugins/Experimental/NeuralRendering/Source/NeuralPostProcessing/Private/NeuralPostProcessModelInstance.cpp
@@ -87,8 +87,7 @@ void UNeuralPostProcessModelInstance::Execute(FRDGBuilder& GraphBuilder)
 			AddPass(GraphBuilder, RDG_EVENT_NAME("EnqueueRDGFence"), [this](FRHICommandListImmediate& RHICmdList)
 				{
 					// Need to wait until the output buffer, RDGOutputBuffer, is ready.
-					RHICmdList.SubmitCommandsAndFlushGPU();
-					RHICmdList.BlockUntilGPUIdle();
+					RHICmdList.SubmitAndBlockUntilGPUIdle();
 				});
 		}
 	};
diff --git a/Engine/Plugins/Experimental/OpenImageDenoise/Source/OpenImageDenoise/Private/Plugin.cpp b/Engine/Plugins/Experimental/OpenImageDenoise/Source/OpenImageDenoise/Private/Plugin.cpp
--- a/Engine/Plugins/Experimental/OpenImageDenoise/Source/OpenImageDenoise/Private/Plugin.cpp
+++ b/Engine/Plugins/Experimental/OpenImageDenoise/Source/OpenImageDenoise/Private/Plugin.cpp
@@ -189,7 +189,7 @@ static void CopyTextureFromGPUToCPU(FRHICommandListImmediate& RHICmdList, FRHITe
 {
 	FRHIGPUTextureReadback Readback(TEXT("DenoiserReadback"));
 	Readback.EnqueueCopy(RHICmdList, SrcTexture, FIntVector::ZeroValue, 0, FIntVector(Size.X, Size.Y, 1));
-	RHICmdList.BlockUntilGPUIdle();
+	RHICmdList.SubmitAndBlockUntilGPUIdle();
 
 	int32_t SrcStride = 0;
 	const PixelType* SrcBuffer = static_cast<PixelType*>(Readback.Lock(SrcStride, nullptr));
diff --git a/Engine/Plugins/Experimental/ProceduralVegetationEditor/Source/ProceduralVegetation/Private/Implementations/PVLightDetection.cpp b/Engine/Plugins/Experimental/ProceduralVegetationEditor/Source/ProceduralVegetation/Private/Implementations/PVLightDetection.cpp
--- a/Engine/Plugins/Experimental/ProceduralVegetationEditor/Source/ProceduralVegetation/Private/Implementations/PVLightDetection.cpp
+++ b/Engine/Plugins/Experimental/ProceduralVegetationEditor/Source/ProceduralVegetation/Private/Implementations/PVLightDetection.cpp
@@ -212,7 +212,7 @@ TArray<FPVPointLightVectorData> FPVLightDetection::ExecuteLightDetection(TArray<
 		AddEnqueueCopyPass(GraphBuilder, Readback, RDGPointLightVectorData, sizeof(FPVPointLightVectorData)* PointLightVectorData.Num());
 
 		GraphBuilder.Execute();
-		RHICmdList.BlockUntilGPUIdle();
+		RHICmdList.SubmitAndBlockUntilGPUIdle();
 
 		Readback->Wait(RHICmdList, RHICmdList.GetGPUMask());
 			
diff --git a/Engine/Plugins/FX/Niagara/Source/Niagara/Private/NiagaraDataInterfaceRenderTarget2D.cpp b/Engine/Plugins/FX/Niagara/Source/Niagara/Private/NiagaraDataInterfaceRenderTarget2D.cpp
--- a/Engine/Plugins/FX/Niagara/Source/Niagara/Private/NiagaraDataInterfaceRenderTarget2D.cpp
+++ b/Engine/Plugins/FX/Niagara/Source/Niagara/Private/NiagaraDataInterfaceRenderTarget2D.cpp
@@ -637,8 +637,8 @@ NIAGARA_API bool UNiagaraDataInterfaceRenderTarget2D::SimCacheWriteFrame(UObject
 					RenderTargetReadback.EnqueueCopy(RHICmdList, InstanceData_RT->TextureRHI, FIntVector(0, 0, 0), 0, FIntVector(RT_VolumeResolution.X, RT_VolumeResolution.Y, 1));
 
 					// Sync the GPU. Unfortunately we can't use the fences because not all RHIs implement them yet.
-					RHICmdList.BlockUntilGPUIdle();
-					RHICmdList.FlushResources();
+					constexpr bool bFlushResources = true;
+					RHICmdList.SubmitAndBlockUntilGPUIdle(bFlushResources);
 
 					//Lock the readback staging texture
 					int32 RowPitchInPixels;
@@ -1071,4 +1071,4 @@ void FNiagaraDataInterfaceProxyRenderTarget2DProxy::GetDispatchArgs(const FNDIGp
 	}
 }
 
-#undef LOCTEXT_NAMESPACE
\ No newline at end of file
+#undef LOCTEXT_NAMESPACE
diff --git a/Engine/Plugins/FX/Niagara/Source/Niagara/Private/NiagaraDataInterfaceRenderTargetVolume.cpp b/Engine/Plugins/FX/Niagara/Source/Niagara/Private/NiagaraDataInterfaceRenderTargetVolume.cpp
--- a/Engine/Plugins/FX/Niagara/Source/Niagara/Private/NiagaraDataInterfaceRenderTargetVolume.cpp
+++ b/Engine/Plugins/FX/Niagara/Source/Niagara/Private/NiagaraDataInterfaceRenderTargetVolume.cpp
@@ -636,8 +636,8 @@ bool UNiagaraDataInterfaceRenderTargetVolume::SimCacheWriteFrame(UObject* Storag
 					RenderTargetReadback.EnqueueCopy(RHICmdList, InstanceData_RT->RenderTarget->GetRHI(), FIntVector(0,0,0), 0, RT_VolumeResolution);
 					
 					// Sync the GPU. Unfortunately we can't use the fences because not all RHIs implement them yet.
-					RHICmdList.BlockUntilGPUIdle();
-					RHICmdList.FlushResources();
+					constexpr bool bFlushResources = true;
+					RHICmdList.SubmitAndBlockUntilGPUIdle(bFlushResources);
 
 					//Lock the readback staging texture
 					int32 RowPitchInPixels;
diff --git a/Engine/Plugins/FX/Niagara/Source/Niagara/Private/NiagaraDataSet.cpp b/Engine/Plugins/FX/Niagara/Source/Niagara/Private/NiagaraDataSet.cpp
--- a/Engine/Plugins/FX/Niagara/Source/Niagara/Private/NiagaraDataSet.cpp
+++ b/Engine/Plugins/FX/Niagara/Source/Niagara/Private/NiagaraDataSet.cpp
@@ -1553,8 +1553,7 @@ void FNiagaraDataBuffer::TransferGPUToCPUImmediate(FRHICommandListImmediate& RHI
 	}
 
 	// Ensure all readbacks are submitted and complete
-	RHICmdList.SubmitCommandsAndFlushGPU();
-	RHICmdList.BlockUntilGPUIdle();
+	RHICmdList.SubmitAndBlockUntilGPUIdle();
 
 	// Read the count and allocato space for results
 	uint32 GPUNumInstances = 0;
diff --git a/Engine/Plugins/FX/Niagara/Source/Niagara/Private/NiagaraGpuReadbackManager.cpp b/Engine/Plugins/FX/Niagara/Source/Niagara/Private/NiagaraGpuReadbackManager.cpp
--- a/Engine/Plugins/FX/Niagara/Source/Niagara/Private/NiagaraGpuReadbackManager.cpp
+++ b/Engine/Plugins/FX/Niagara/Source/Niagara/Private/NiagaraGpuReadbackManager.cpp
@@ -61,8 +61,7 @@ void FNiagaraGpuReadbackManager::WaitCompletion(FRHICommandListImmediate& RHICmd
 {
 	// Ensure all GPU commands have been executed as we will ignore the fence
 	// This is because the fence may be implemented as a simple counter rather than a real fence
-	RHICmdList.SubmitCommandsAndFlushGPU();
-	RHICmdList.BlockUntilGPUIdle();
+	RHICmdList.SubmitAndBlockUntilGPUIdle();
 
 	// Perform a tick which will flush everything
 	TickInternal(RHICmdList, true);
diff --git a/Engine/Plugins/FX/Niagara/Source/NiagaraEditor/Private/NiagaraBakerRendererOutputSparseVolumeTexture.cpp b/Engine/Plugins/FX/Niagara/Source/NiagaraEditor/Private/NiagaraBakerRendererOutputSparseVolumeTexture.cpp
--- a/Engine/Plugins/FX/Niagara/Source/NiagaraEditor/Private/NiagaraBakerRendererOutputSparseVolumeTexture.cpp
+++ b/Engine/Plugins/FX/Niagara/Source/NiagaraEditor/Private/NiagaraBakerRendererOutputSparseVolumeTexture.cpp
@@ -267,8 +267,8 @@ void FNiagaraBakerRendererOutputSparseVolumeTexture::BakeFrame(FNiagaraBakerFeed
 				}
 
 				// Sync the GPU. Unfortunately we can't use the fences because not all RHIs implement them yet.
-				RHICmdList.BlockUntilGPUIdle();
-				RHICmdList.FlushResources();
+				constexpr bool bFlushResources = true;
+				RHICmdList.SubmitAndBlockUntilGPUIdle(bFlushResources);
 
 				//Lock the readback staging texture
 				int32 RowPitchInPixels;
@@ -584,7 +584,7 @@ void FNiagaraBakerRendererOutputSparseVolumeTexture::EndBake(FNiagaraBakerFeedba
 						// Execute the graph.
 						GraphBuilder.Execute();
 
-						RHICmdList.BlockUntilGPUIdle();
+						RHICmdList.SubmitAndBlockUntilGPUIdle();
 
 						check(RenderTargetReadback.IsReady());
 
diff --git a/Engine/Plugins/Media/MediaIOFramework/Source/MediaIOCore/Private/MediaIOFrameManager.h b/Engine/Plugins/Media/MediaIOFramework/Source/MediaIOCore/Private/MediaIOFrameManager.h
--- a/Engine/Plugins/Media/MediaIOFramework/Source/MediaIOCore/Private/MediaIOFrameManager.h
+++ b/Engine/Plugins/Media/MediaIOFramework/Source/MediaIOCore/Private/MediaIOFrameManager.h
@@ -402,7 +402,7 @@ namespace UE::MediaCaptureData
 			if (ReadbackBuffer->IsReady() == false)
 			{
 				UE_LOGF(LogMediaIOCore, Verbose, "Fence for buffer readback was not ready, blocking.");
-				RHICmdList.BlockUntilGPUIdle();
+				RHICmdList.SubmitAndBlockUntilGPUIdle();
 			}
 
 			OutRowStride = GetBufferResource()->GetStride();
diff --git a/Engine/Plugins/NNE/NNEDenoiser/Source/NNEDenoiser/Private/NNEDenoiserUnitTestOidn.cpp b/Engine/Plugins/NNE/NNEDenoiser/Source/NNEDenoiser/Private/NNEDenoiserUnitTestOidn.cpp
--- a/Engine/Plugins/NNE/NNEDenoiser/Source/NNEDenoiser/Private/NNEDenoiserUnitTestOidn.cpp
+++ b/Engine/Plugins/NNE/NNEDenoiser/Source/NNEDenoiser/Private/NNEDenoiserUnitTestOidn.cpp
@@ -121,7 +121,7 @@ bool FTransferFunctionTest::RunTest(const FString& Parameter)
 
 			GraphBuilder.Execute();
 
-			RHICmdList.BlockUntilGPUIdle();
+			RHICmdList.SubmitAndBlockUntilGPUIdle();
 
 			Signal->Trigger();
 		}
diff --git a/Engine/Plugins/NNE/NNEDenoiser/Source/NNEDenoiser/Private/NNEDenoiserUnitTests.cpp b/Engine/Plugins/NNE/NNEDenoiser/Source/NNEDenoiser/Private/NNEDenoiserUnitTests.cpp
--- a/Engine/Plugins/NNE/NNEDenoiser/Source/NNEDenoiser/Private/NNEDenoiserUnitTests.cpp
+++ b/Engine/Plugins/NNE/NNEDenoiser/Source/NNEDenoiser/Private/NNEDenoiserUnitTests.cpp
@@ -96,7 +96,7 @@ bool FAutoExposureTest::RunTest(const FString& Parameter)
 
 			GraphBuilder.Execute();
 
-			RHICmdList.BlockUntilGPUIdle();
+			RHICmdList.SubmitAndBlockUntilGPUIdle();
 
 			Signal->Trigger();
 		}
diff --git a/Engine/Plugins/NNE/NNEDenoiser/Source/NNEDenoiser/Private/NNEDenoiserUtils.h b/Engine/Plugins/NNE/NNEDenoiser/Source/NNEDenoiser/Private/NNEDenoiserUtils.h
--- a/Engine/Plugins/NNE/NNEDenoiser/Source/NNEDenoiser/Private/NNEDenoiserUtils.h
+++ b/Engine/Plugins/NNE/NNEDenoiser/Source/NNEDenoiser/Private/NNEDenoiserUtils.h
@@ -58,7 +58,7 @@ namespace UE::NNEDenoiser::Private
 		FRHIGPUTextureReadback Readback(TEXT("NNEDenoiser.CopyTextureFromGPUToCPU"));
 		Readback.EnqueueCopy(RHICmdList, SrcTexture, FIntVector::ZeroValue, 0, FIntVector(Size.X, Size.Y, 1));
 
-		RHICmdList.BlockUntilGPUIdle();
+		RHICmdList.SubmitAndBlockUntilGPUIdle();
 
 		int32_t SrcStride = 0;
 		const PixelType* SrcBuffer = static_cast<PixelType*>(Readback.Lock(SrcStride, nullptr));
@@ -98,7 +98,7 @@ namespace UE::NNEDenoiser::Private
 		FRHIGPUBufferReadback Readback(TEXT("NNEDenoiser.CopyBufferFromGPUToCPU"));
 		Readback.EnqueueCopy(RHICmdList, SrcBuffer, NumBytes);
 
-		RHICmdList.BlockUntilGPUIdle();
+		RHICmdList.SubmitAndBlockUntilGPUIdle();
 
 		const ElementType* Src = static_cast<ElementType*>(Readback.Lock(NumBytes));
 		ElementType* Dst = DstArray.GetData();
@@ -165,4 +165,4 @@ namespace UE::NNEDenoiser::Private
 		return EDataType::None;
 	} 
 
-}
\ No newline at end of file
+}
diff --git a/Engine/Plugins/Tests/RHITests/Source/RHITests/Private/RHIBindlessTests.cpp b/Engine/Plugins/Tests/RHITests/Source/RHITests/Private/RHIBindlessTests.cpp
--- a/Engine/Plugins/Tests/RHITests/Source/RHITests/Private/RHIBindlessTests.cpp
+++ b/Engine/Plugins/Tests/RHITests/Source/RHITests/Private/RHIBindlessTests.cpp
@@ -170,8 +170,7 @@ static bool ReadbackAndVerifyColors(
 	FRHIGPUBufferReadback OutputBufferReadback(TEXT("OutputUAVReadback"));
 	OutputBufferReadback.EnqueueCopy(RHICmdList, OutputBuffer, OutputBufferSize);
 
-	RHICmdList.BlockUntilGPUIdle();
-	RHICmdList.FlushResources();
+	RHICmdList.SubmitAndBlockUntilGPUIdle();
 
 	bool bExpectedResult = true;
 
@@ -456,8 +455,7 @@ bool RHIBindlessTests::Test_DescriptorRange_UAV(FRHICommandListImmediate& RHICmd
 		}
 	}
 
-	RHICmdList.BlockUntilGPUIdle();
-	RHICmdList.FlushResources();
+	RHICmdList.SubmitAndBlockUntilGPUIdle();
 
 	// Readback each output buffer and verify.
 	bool bExpectedResult = true;
@@ -466,8 +464,7 @@ bool RHIBindlessTests::Test_DescriptorRange_UAV(FRHICommandListImmediate& RHICmd
 	{
 		FRHIGPUBufferReadback Readback(TEXT("UAVRangeReadback"));
 		Readback.EnqueueCopy(RHICmdList, OutputBuffers[Index], sizeof(FLinearColor));
-		RHICmdList.BlockUntilGPUIdle();
-		RHICmdList.FlushResources();
+		RHICmdList.SubmitAndBlockUntilGPUIdle();
 
 		const FLinearColor* OutputData = reinterpret_cast<const FLinearColor*>(Readback.Lock(sizeof(FLinearColor)));
 		const FLinearColor& ExpectedColor = InputValues[Index];
diff --git a/Engine/Plugins/Tests/RHITests/Source/RHITests/Private/RHIBufferTests.cpp b/Engine/Plugins/Tests/RHITests/Source/RHITests/Private/RHIBufferTests.cpp
--- a/Engine/Plugins/Tests/RHITests/Source/RHITests/Private/RHIBufferTests.cpp
+++ b/Engine/Plugins/Tests/RHITests/Source/RHITests/Private/RHIBufferTests.cpp
@@ -14,8 +14,8 @@ bool FRHIBufferTests::VerifyBufferContents(const TCHAR* TestName, FRHICommandLis
 	}
 
 	// @todo - readback API is inconsistent across RHIs
-	RHICmdList.BlockUntilGPUIdle();
-	RHICmdList.FlushResources();
+	constexpr bool bFlushResources = true;
+	RHICmdList.SubmitAndBlockUntilGPUIdle(bFlushResources);
 
 	ON_SCOPE_EXIT
 	{
@@ -75,8 +75,8 @@ bool FRHIBufferTests::VerifyBufferContents(const TCHAR* TestName, FRHICommandLis
 		RHICmdList.CopyToStagingBuffer(Buffer, StagingBuffer, 0, NumBytes);
 
 		// @todo - readback API is inconsistent across RHIs
-		RHICmdList.BlockUntilGPUIdle();
-		RHICmdList.FlushResources();
+		constexpr bool bFlushResources = true;
+		RHICmdList.SubmitAndBlockUntilGPUIdle(bFlushResources);
 
 		void* Memory = RHICmdList.LockStagingBuffer(StagingBuffer, 0, NumBytes);
 
diff --git a/Engine/Plugins/Tests/RHITests/Source/RHITests/Private/RHIClearTextureTests.cpp b/Engine/Plugins/Tests/RHITests/Source/RHITests/Private/RHIClearTextureTests.cpp
--- a/Engine/Plugins/Tests/RHITests/Source/RHITests/Private/RHIClearTextureTests.cpp
+++ b/Engine/Plugins/Tests/RHITests/Source/RHITests/Private/RHIClearTextureTests.cpp
@@ -1097,7 +1097,7 @@ static void BlockGPUAndLaunchTestTasks(FRHICommandListImmediate& RHICmdList, con
 	{
 		// Also flushes the RHI thread
 		CLEAR_TEST_SCOPED_NAMED_EVENT_TEXT("Test_ClearTexture_BlockUntilGPUIdle", FColor::Magenta);
-		RHICmdList.BlockUntilGPUIdle();
+		RHICmdList.SubmitAndBlockUntilGPUIdle();
 	}
 
 	for (FTestOperation& TestOperation : TestOperationGroup)
@@ -1210,7 +1210,7 @@ static void FillClearAndTest(FRHICommandListImmediate& RHICmdList, FRHITexture*
 			}
 		}
 
-		RHICmdList.BlockUntilGPUIdle();
+		RHICmdList.SubmitAndBlockUntilGPUIdle();
 
 		MapStagingSurfaces(RHICmdList, TestOperation.DrawResult);
 
@@ -1360,7 +1360,7 @@ static bool TestDrawAtlasing(FRHICommandListImmediate& RHICmdList, FVertexDeclar
 	RHICmdList.Transition(FRHITransitionInfo(SourceTexture, ERHIAccess::CopySrc, ERHIAccess::SRVCompute));
 	CopyTextureToUnpackedStaging(RHICmdList, SourceTexture, SourceTextureSRV, TmpTexture2DAtlas, TmpTexture2DAtlasUAV, DestTextureStaging, VertexDeclarationRHI);
 
-	RHICmdList.BlockUntilGPUIdle();
+	RHICmdList.SubmitAndBlockUntilGPUIdle();
 	FStagingData DrawResult{ DestTextureStaging };
 	MapStagingSurfaces(RHICmdList, DrawResult);
 	bool bTestSuccess = VerifyTextureData(DrawResult, SourceTexture, CheckTestData);
@@ -1372,7 +1372,6 @@ bool FRHIClearTextureTests::Test_ClearTexture(FRHICommandListImmediate& RHICmdLi
 {
 	CLEAR_TEST_SCOPED_NAMED_EVENT_TEXT("Test_ClearTexture", FColor::Magenta);
 
-	RHICmdList.FlushResources();
 	RHICmdList.ImmediateFlush(EImmediateFlushType::FlushRHIThreadFlushResources);
 
 	FVertexDeclarationElementList VertexDeclarationElements;
@@ -1421,7 +1420,6 @@ bool FRHIClearTextureTests::Test_ClearTexture(FRHICommandListImmediate& RHICmdLi
 			bTestDrawAtlasing &= TestDrawAtlasing(RHICmdList, VertexDeclarationRHI,
 				FRHITextureCreateDesc::Create3D(TEXT("SourceTexture"), 32, 32, 5, PixelFormat).SetNumMips(1).SetInitialState(ERHIAccess::CopySrc));
 
-			RHICmdList.FlushResources();
 			RHICmdList.ImmediateFlush(EImmediateFlushType::FlushRHIThreadFlushResources);
 		}
 	}
@@ -1563,7 +1561,6 @@ bool FRHIClearTextureTests::Test_ClearTexture(FRHICommandListImmediate& RHICmdLi
 
 			{
 				CLEAR_TEST_SCOPED_NAMED_EVENT_TEXT(TEXT("FlushResources"), FColor::Magenta);
-				RHICmdList.FlushResources();
 				RHICmdList.ImmediateFlush(EImmediateFlushType::FlushRHIThreadFlushResources);
 			}
 
@@ -1624,7 +1621,6 @@ bool FRHIClearTextureTests::Test_ClearTexture(FRHICommandListImmediate& RHICmdLi
 
 	}
 
-	RHICmdList.FlushResources();
 	RHICmdList.ImmediateFlush(EImmediateFlushType::FlushRHIThreadFlushResources);
 
 	return bResult;
diff --git a/Engine/Plugins/Tests/RHITests/Source/RHITests/Private/RHIReadbackTests.cpp b/Engine/Plugins/Tests/RHITests/Source/RHITests/Private/RHIReadbackTests.cpp
--- a/Engine/Plugins/Tests/RHITests/Source/RHITests/Private/RHIReadbackTests.cpp
+++ b/Engine/Plugins/Tests/RHITests/Source/RHITests/Private/RHIReadbackTests.cpp
@@ -149,8 +149,8 @@ bool FRHIReadbackTests::Test_BufferReadback(FRHICommandListImmediate& RHICmdList
 	StructuredBufferReadback.EnqueueCopy(RHICmdList, StructuredBuffer, StructuredBufferSize);
 
 	// Sync the GPU. Unfortunately we can't use the fences because not all RHIs implement them yet.
-	RHICmdList.BlockUntilGPUIdle();
-	RHICmdList.FlushResources();
+	constexpr bool bFlushResources = true;
+	RHICmdList.SubmitAndBlockUntilGPUIdle(bFlushResources);
 
 	// Check the contents of the typed buffer.
 	bool bTypedBufferResult = true;
@@ -336,8 +336,8 @@ bool FRHIReadbackTests::Test_TextureReadback(FRHICommandListImmediate& RHICmdLis
 	TextureArrayReadback.EnqueueCopy(RHICmdList, TextureArray, ReadbackOffsetArray, ReadbackArraySlice, ReadbackSizeArray);
 
 	// Sync the GPU. Unfortunately we can't use the fences because not all RHIs implement them yet.
-	RHICmdList.BlockUntilGPUIdle();
-	RHICmdList.FlushResources();
+	constexpr bool bFlushResources = true;
+	RHICmdList.SubmitAndBlockUntilGPUIdle(bFlushResources);
 
 	// Check the contents of the 2D texture.
 	bool bTexture2DResult = true;
diff --git a/Engine/Plugins/Tests/RHITests/Source/RHITests/Private/RHIRootConstantsTests.cpp b/Engine/Plugins/Tests/RHITests/Source/RHITests/Private/RHIRootConstantsTests.cpp
--- a/Engine/Plugins/Tests/RHITests/Source/RHITests/Private/RHIRootConstantsTests.cpp
+++ b/Engine/Plugins/Tests/RHITests/Source/RHITests/Private/RHIRootConstantsTests.cpp
@@ -85,8 +85,8 @@ static bool ReadbackAndVerifyRootConstants(
 	FRHIGPUBufferReadback Readback(TestName);
 	Readback.EnqueueCopy(RHICmdList, OutputBuffer, BufferSize);
 
-	RHICmdList.BlockUntilGPUIdle();
-	RHICmdList.FlushResources();
+	constexpr bool bFlushResources = true;
+	RHICmdList.SubmitAndBlockUntilGPUIdle(bFlushResources);
 
 	const FUint32Vector4* OutputData = reinterpret_cast<const FUint32Vector4*>(Readback.Lock(BufferSize));
 
diff --git a/Engine/Plugins/Tests/RHITests/Source/RHITests/Private/RHITextureTests.cpp b/Engine/Plugins/Tests/RHITests/Source/RHITests/Private/RHITextureTests.cpp
--- a/Engine/Plugins/Tests/RHITests/Source/RHITests/Private/RHITextureTests.cpp
+++ b/Engine/Plugins/Tests/RHITests/Source/RHITests/Private/RHITextureTests.cpp
@@ -47,8 +47,9 @@ bool FRHITextureTests::VerifyTextureContents(const TCHAR* TestName, FRHICommandL
 					FGPUFenceRHIRef GPUFence = RHICreateGPUFence(TEXT("ReadbackFence"));
 					RHICmdList.WriteGPUFence(GPUFence);
 
-					RHICmdList.SubmitCommandsAndFlushGPU(); // @todo - refactor RHI readback API. This shouldn't be necessary
-					RHICmdList.FlushResources();
+					// @todo - refactor RHI readback API. This shouldn't be necessary
+					constexpr bool bFlushResources = true;
+					RHICmdList.SubmitAndBlockUntilGPUIdle(bFlushResources);
 
 					int32 Width, Height;
 					void* Ptr;
@@ -74,4 +75,4 @@ bool FRHITextureTests::VerifyTextureContents(const TCHAR* TestName, FRHICommandL
 	}
 
 	return bResult;
-}
\ No newline at end of file
+}
diff --git a/Engine/Plugins/Tests/RHITests/Source/RHITests/Public/RHITextureTests.h b/Engine/Plugins/Tests/RHITests/Source/RHITests/Public/RHITextureTests.h
--- a/Engine/Plugins/Tests/RHITests/Source/RHITests/Public/RHITextureTests.h
+++ b/Engine/Plugins/Tests/RHITests/Source/RHITests/Public/RHITextureTests.h
@@ -1141,7 +1141,7 @@ public:
 				FResolveRect SourceRect = FResolveRect(PatternOffset.X, PatternOffset.Y, PatternOffset.X + PatternSize.X, PatternOffset.Y + PatternSize.Y);
 
 				Readback.EnqueueCopy(RHICmdList, SourceTexture, SourceRect);
-				RHICmdList.SubmitCommandsAndFlushGPU();
+				RHICmdList.SubmitAndBlockUntilGPUIdle();
 
 				int32 Pitch;
 				void* ReadbackData = Readback.Lock(Pitch);
@@ -1440,7 +1440,7 @@ public:
 				}
 				RHICmdList.UnlockTexture2D(SourceTexture, 0, false);
 
-				RHICmdList.SubmitCommandsAndFlushGPU();
+				RHICmdList.SubmitAndBlockUntilGPUIdle();
 
 				TArray<FColor> OutData;
 				FReadSurfaceDataFlags Flags;
@@ -1613,8 +1613,8 @@ public:
 			FGPUFenceRHIRef GPUFence = RHICreateGPUFence(TEXT("ReadbackFence"));
 			RHICmdList.WriteGPUFence(GPUFence);
 
-			RHICmdList.SubmitCommandsAndFlushGPU();
-			RHICmdList.FlushResources();
+			constexpr bool bFlushResources = true;
+			RHICmdList.SubmitAndBlockUntilGPUIdle(bFlushResources);
 
 			int32 SingleLockStagingSurfaceWidth, SingleLockStagingSurfaceHeight;
 			void* SingleLockData;
diff --git a/Engine/Plugins/TextureGraph/Source/TextureGraphEngine/2D/TextureHelper.cpp b/Engine/Plugins/TextureGraph/Source/TextureGraphEngine/2D/TextureHelper.cpp
--- a/Engine/Plugins/TextureGraph/Source/TextureGraphEngine/2D/TextureHelper.cpp
+++ b/Engine/Plugins/TextureGraph/Source/TextureGraphEngine/2D/TextureHelper.cpp
@@ -849,8 +849,7 @@ RawBufferPtr TextureHelper::RawFromResource(const FTextureRHIRef& ResourceRHI, c
 		const uint32 BytesPerPixel = BitsPerPixel / 8;
 
 		const TUniquePtr<FRHIGPUTextureReadback> TextureReadback = MakeUnique<FRHIGPUTextureReadback>(TEXT("RawFromResourceTextureReadback"));
-		RHICmdList.FlushResources();
-		RHICmdList.ImmediateFlush(EImmediateFlushType::WaitForOutstandingTasksOnly);
+		RHICmdList.ImmediateFlush(EImmediateFlushType::FlushRHIThreadFlushResources);
 
 		RHICmdList.Transition(FRHITransitionInfo(ResourceRHI, ERHIAccess::Unknown, ERHIAccess::CopySrc));
 
@@ -858,7 +857,7 @@ RawBufferPtr TextureHelper::RawFromResource(const FTextureRHIRef& ResourceRHI, c
 
 		RHICmdList.Transition(FRHITransitionInfo(ResourceRHI, ERHIAccess::CopySrc, ERHIAccess::SRVMask));
 
-		RHICmdList.BlockUntilGPUIdle();
+		RHICmdList.SubmitAndBlockUntilGPUIdle();
 
 		//check(TextureReadback->IsReady());
 		{
diff --git a/Engine/Plugins/VirtualProduction/CameraCalibrationCore/Source/CameraCalibrationCore/Private/DistortionRenderingUtils.cpp b/Engine/Plugins/VirtualProduction/CameraCalibrationCore/Source/CameraCalibrationCore/Private/DistortionRenderingUtils.cpp
--- a/Engine/Plugins/VirtualProduction/CameraCalibrationCore/Source/CameraCalibrationCore/Private/DistortionRenderingUtils.cpp
+++ b/Engine/Plugins/VirtualProduction/CameraCalibrationCore/Source/CameraCalibrationCore/Private/DistortionRenderingUtils.cpp
@@ -114,8 +114,7 @@ namespace DistortionRenderingUtils
 				RHICmdList.CopyToStagingBuffer(UndistortedPointsBuffer.Buffer, DestinationStagingBuffer, 0, BufferSize);
 
 				// Wait to ensure that the staging buffer is ready to read
-				RHICmdList.SubmitCommandsAndFlushGPU();
-				RHICmdList.BlockUntilGPUIdle();
+				RHICmdList.SubmitAndBlockUntilGPUIdle();
 
 				// Copy data out of the staging buffer
 				FVector2f* UndistortedPointData = static_cast<FVector2f*>(DestinationStagingBuffer->Lock(0, BufferSize));
@@ -143,4 +142,4 @@ namespace DistortionRenderingUtils
 		// Ensure that all rendering commands have been issued so that OutUndistortedPoints() is valid and has correct data 
 		FlushRenderingCommands();
 	}
-}
\ No newline at end of file
+}
diff --git a/Engine/Source/Runtime/Apple/MetalRHI/Private/MetalDynamicRHI.h b/Engine/Source/Runtime/Apple/MetalRHI/Private/MetalDynamicRHI.h
--- a/Engine/Source/Runtime/Apple/MetalRHI/Private/MetalDynamicRHI.h
+++ b/Engine/Source/Runtime/Apple/MetalRHI/Private/MetalDynamicRHI.h
@@ -176,7 +176,6 @@ public:
 	virtual bool RHIGetRenderQueryResult(FRHIRenderQuery* RenderQuery, uint64& OutResult, bool bWait, uint32 GPUIndex = INDEX_NONE) final override;
 	virtual FTextureRHIRef RHIGetViewportBackBuffer(FRHIViewport* Viewport) final override;
 	virtual FTextureRHIRef RHIGetViewportBackBufferAtLayerIndex(FRHIViewport* Viewport, uint32 DrawLayerIndex) final override;
-	virtual void RHIFlushResources() final override;
 	virtual FViewportRHIRef RHICreateViewport(void* WindowHandle, uint32 SizeX, uint32 SizeY, bool bIsFullscreen, EPixelFormat PixelFormat) final override;
 	virtual void RHIResizeViewport(FRHIViewport* Viewport, uint32 SizeX, uint32 SizeY, bool bIsFullscreen, EPixelFormat PixelFormat) final override;
 	virtual void RHIEndDrawingViewport(FRHICommandListImmediate& RHICmdList, FRHIViewport* Viewport, const FRHIPresentArgs& PresentArgs) final override;
diff --git a/Engine/Source/Runtime/Apple/MetalRHI/Private/MetalRHI.cpp b/Engine/Source/Runtime/Apple/MetalRHI/Private/MetalRHI.cpp
--- a/Engine/Source/Runtime/Apple/MetalRHI/Private/MetalRHI.cpp
+++ b/Engine/Source/Runtime/Apple/MetalRHI/Private/MetalRHI.cpp
@@ -1434,12 +1434,6 @@ bool FMetalDynamicRHI::RHIGetAvailableResolutions(FScreenResolutionArray& Resolu
 #endif
 }
 
-void FMetalDynamicRHI::RHIFlushResources()
-{
-    MTL_SCOPED_AUTORELEASE_POOL;
-	Device->DrainHeap();
-}
-
 void* FMetalDynamicRHI::RHIGetNativeDevice()
 {
 	return (void*)Device->GetDevice();
@@ -1633,6 +1627,8 @@ void FMetalDynamicRHI::RHIProcessDeleteQueue()
 				}
 			}
 		}
+		
+		Device->DrainHeap();
 	});
 }
 
diff --git a/Engine/Source/Runtime/Apple/MetalRHI/Private/MetalViewport.cpp b/Engine/Source/Runtime/Apple/MetalRHI/Private/MetalViewport.cpp
--- a/Engine/Source/Runtime/Apple/MetalRHI/Private/MetalViewport.cpp
+++ b/Engine/Source/Runtime/Apple/MetalRHI/Private/MetalViewport.cpp
@@ -291,7 +291,7 @@ void FMetalViewport::DeleteNativeLayer(int32 OldNativeLayer)
 	ENQUEUE_RENDER_COMMAND(FlushPendingRHICommands)(
 		[Viewport = this, OldNativeLayer](FRHICommandListImmediate& RHICmdList)
 		{
-			GRHICommandList.GetImmediateCommandList().BlockUntilGPUIdle();
+			GRHICommandList.GetImmediateCommandList().SubmitAndBlockUntilGPUIdle();
 			dispatch_sync(dispatch_get_main_queue(), ^
 		    {
 				FScopeLock Lock(&Viewport->Mutex);
@@ -517,7 +517,7 @@ void FMetalViewport::Resize(uint32 InSizeX, uint32 InSizeY, bool bInIsFullscreen
     ENQUEUE_RENDER_COMMAND(FlushPendingRHICommands)(
         [Viewport = this](FRHICommandListImmediate& RHICmdList)
         {
-            GRHICommandList.GetImmediateCommandList().BlockUntilGPUIdle();
+            GRHICommandList.GetImmediateCommandList().SubmitAndBlockUntilGPUIdle();
         });
 	
 #if WITH_MULTI_VIEW_SLATE_WINDOW_SUPPORT
diff --git a/Engine/Source/Runtime/D3D12RHI/Private/D3D12RHI.cpp b/Engine/Source/Runtime/D3D12RHI/Private/D3D12RHI.cpp
--- a/Engine/Source/Runtime/D3D12RHI/Private/D3D12RHI.cpp
+++ b/Engine/Source/Runtime/D3D12RHI/Private/D3D12RHI.cpp
@@ -402,11 +402,6 @@ IRHICommandContext* FD3D12DynamicRHI::RHIGetDefaultContext()
 	return DefaultCommandContext;
 }
 
-void FD3D12DynamicRHI::RHIFlushResources()
-{
-	// Nothing to do (yet!)
-}
-
 void FD3D12DynamicRHI::EnqueueEndOfPipeTask(TUniqueFunction<void()> TaskFunc, TUniqueFunction<void(FD3D12Payload&)> ModifyPayloadCallback)
 {
 	SCOPED_NAMED_EVENT_TEXT("EnqueueEndOfPipeTask", FColor::Yellow);
diff --git a/Engine/Source/Runtime/D3D12RHI/Private/D3D12RHIPrivate.h b/Engine/Source/Runtime/D3D12RHI/Private/D3D12RHIPrivate.h
--- a/Engine/Source/Runtime/D3D12RHI/Private/D3D12RHIPrivate.h
+++ b/Engine/Source/Runtime/D3D12RHI/Private/D3D12RHIPrivate.h
@@ -408,7 +408,6 @@ public:
 	virtual uint64 RHIComputeStatePrecachePSOHash(const FGraphicsPipelineStateInitializer& Initializer) final override;
 	virtual uint64 RHIComputePrecachePSOHash(const FGraphicsPipelineStateInitializer& Initializer) final override;
 	virtual bool RHIMatchPrecachePSOInitializers(const FGraphicsPipelineStateInitializer& LHS, const FGraphicsPipelineStateInitializer& RHS) final override;
-	virtual void RHIFlushResources() final override;
 	virtual FViewportRHIRef RHICreateViewport(void* WindowHandle, uint32 SizeX, uint32 SizeY, bool bIsFullscreen, EPixelFormat PixelFormat) final override;
 	virtual void RHIResizeViewport(FRHIViewport* ViewportRHI, uint32 SizeX, uint32 SizeY, bool bIsFullscreen, EPixelFormat PixelFormat) final override;
 	virtual void RHIEndDrawingViewport(FRHICommandListImmediate& RHICmdList, FRHIViewport* Viewport, const FRHIPresentArgs& PresentArgs) final override;
diff --git a/Engine/Source/Runtime/Engine/Private/GameEngine.cpp b/Engine/Source/Runtime/Engine/Private/GameEngine.cpp
--- a/Engine/Source/Runtime/Engine/Private/GameEngine.cpp
+++ b/Engine/Source/Runtime/Engine/Private/GameEngine.cpp
@@ -2076,7 +2076,7 @@ void UGameEngine::Tick( float DeltaSeconds, bool bIdleMode )
 			{
 				ENQUEUE_RENDER_COMMAND(SubmitAndBlockUntilGPUIdle_MinimizedRealtime)([](FRHICommandListImmediate& RHICmdList)
 				{
-					RHICmdList.BlockUntilGPUIdle();
+					RHICmdList.SubmitAndBlockUntilGPUIdle();
 				});
 			}
 			else
diff --git a/Engine/Source/Runtime/Engine/Private/Streaming/StreamingManagerTexture.cpp b/Engine/Source/Runtime/Engine/Private/Streaming/StreamingManagerTexture.cpp
--- a/Engine/Source/Runtime/Engine/Private/Streaming/StreamingManagerTexture.cpp
+++ b/Engine/Source/Runtime/Engine/Private/Streaming/StreamingManagerTexture.cpp
@@ -433,7 +433,6 @@ bool FRenderAssetStreamingManager::StreamOutRenderAssetData( int64 RequiredMemor
 					[](FRHICommandListImmediate& RHICmdList)
 					{				
 						RHICmdList.ImmediateFlush(EImmediateFlushType::FlushRHIThreadFlushResources);
-						RHICmdList.FlushResources();
 					});
 				FlushRenderingCommands();
 				TempMemoryUsed = 0;
@@ -3247,4 +3246,4 @@ float FRenderAssetStreamingManagerHelper::GetWorldTime(UWorld* InWorld)
 	}
 
 	return 0.f;
-}
\ No newline at end of file
+}
diff --git a/Engine/Source/Runtime/Engine/Private/UnrealEngine.cpp b/Engine/Source/Runtime/Engine/Private/UnrealEngine.cpp
--- a/Engine/Source/Runtime/Engine/Private/UnrealEngine.cpp
+++ b/Engine/Source/Runtime/Engine/Private/UnrealEngine.cpp
@@ -16733,7 +16733,6 @@ void UEngine::TrimMemory()
 #if !WITH_EDITOR
 			GetRendererModule().ReleaseVirtualTexturePendingResources();
 #endif	
-			RHIFlushResources();
 			RHICmdList.ImmediateFlush(EImmediateFlushType::FlushRHIThreadFlushResources);
 		});
 	FlushRenderingCommands();
diff --git a/Engine/Source/Runtime/Landscape/Private/LandscapeUtilsPrivate.cpp b/Engine/Source/Runtime/Landscape/Private/LandscapeUtilsPrivate.cpp
--- a/Engine/Source/Runtime/Landscape/Private/LandscapeUtilsPrivate.cpp
+++ b/Engine/Source/Runtime/Landscape/Private/LandscapeUtilsPrivate.cpp
@@ -193,16 +193,15 @@ void SubmitGPUCommands(bool bBlockUntilRTComplete, bool bBlockRTUntilGPUComplete
 	ENQUEUE_RENDER_COMMAND(FFlushResourcesCommand)(
 		[ResultsReadyEvent, bBlockRTUntilGPUComplete](FRHICommandListImmediate& RHICmdList)
 		{
-			RHICmdList.ImmediateFlush(EImmediateFlushType::FlushRHIThreadFlushResources);
-			RHICmdList.FlushResources();
-			RHICmdList.SubmitCommandsAndFlushGPU();
+			constexpr bool bFlushResources = true;
+			RHICmdList.SubmitAndBlockUntilGPUIdle(bFlushResources);
 
 			if (ResultsReadyEvent)
 			{
 				if (bBlockRTUntilGPUComplete)
 				{
 					// Block render thread waiting for GPU to complete.  Note this can be very expensive on some platforms.
-					RHICmdList.BlockUntilGPUIdle();
+					RHICmdList.SubmitAndBlockUntilGPUIdle();
 				}
 
 				ResultsReadyEvent->Trigger();
@@ -251,4 +250,4 @@ uint32 ComputeNamedGrassTypesHash(const TMap<FName, TObjectPtr<ULandscapeGrassTy
 	return Hash;
 }
 
-} // namespace UE::Landscape::Grass::Private
\ No newline at end of file
+} // namespace UE::Landscape::Grass::Private
diff --git a/Engine/Source/Runtime/NullDrv/Public/NullRHI.h b/Engine/Source/Runtime/NullDrv/Public/NullRHI.h
--- a/Engine/Source/Runtime/NullDrv/Public/NullRHI.h
+++ b/Engine/Source/Runtime/NullDrv/Public/NullRHI.h
@@ -429,11 +429,6 @@ public:
 
 	}
 
-	virtual void RHIFlushResources() final override
-	{
-
-	}
-
 	virtual FViewportRHIRef RHICreateViewport(void* WindowHandle, uint32 SizeX, uint32 SizeY, bool bIsFullscreen, EPixelFormat PixelFormat) final override
 	{ 
 		return new FRHIViewport(); 
diff --git a/Engine/Source/Runtime/OpenGLDrv/Private/Linux/LinuxOpenGLPlatform.cpp b/Engine/Source/Runtime/OpenGLDrv/Private/Linux/LinuxOpenGLPlatform.cpp
--- a/Engine/Source/Runtime/OpenGLDrv/Private/Linux/LinuxOpenGLPlatform.cpp
+++ b/Engine/Source/Runtime/OpenGLDrv/Private/Linux/LinuxOpenGLPlatform.cpp
@@ -516,11 +516,6 @@ bool PlatformBlitToViewport(IRHICommandContext& RHICmdContext,
 	return true;
 }
 
-void PlatformFlushIfNeeded()
-{
-	glFinish();
-}
-
 void PlatformRenderingContextSetup(FPlatformOpenGLDevice* Device)
 {
 	check(Device && Device->RenderingContext.hWnd && Device->RenderingContext.hGLContext);
diff --git a/Engine/Source/Runtime/OpenGLDrv/Private/OpenGLDevice.cpp b/Engine/Source/Runtime/OpenGLDrv/Private/OpenGLDevice.cpp
--- a/Engine/Source/Runtime/OpenGLDrv/Private/OpenGLDevice.cpp
+++ b/Engine/Source/Runtime/OpenGLDrv/Private/OpenGLDevice.cpp
@@ -1394,11 +1394,6 @@ void FOpenGLDynamicRHI::Cleanup()
 	ContextState.CleanupResources();
 }
 
-void FOpenGLDynamicRHI::RHIFlushResources()
-{
-	PlatformFlushIfNeeded();
-}
-
 void FOpenGLDynamicRHI::RHIAcquireThreadOwnership()
 {
 	PlatformRenderingContextSetup(PlatformDevice);
diff --git a/Engine/Source/Runtime/OpenGLDrv/Private/OpenGLDrv.h b/Engine/Source/Runtime/OpenGLDrv/Private/OpenGLDrv.h
--- a/Engine/Source/Runtime/OpenGLDrv/Private/OpenGLDrv.h
+++ b/Engine/Source/Runtime/OpenGLDrv/Private/OpenGLDrv.h
@@ -161,7 +161,6 @@ public:
 	UE_API virtual FTextureRHIRef RHICreateAliasedTexture(FTextureRHIRef& SourceTexture) final override;
 	UE_API virtual void RHIAcquireThreadOwnership() final override;
 	UE_API virtual void RHIReleaseThreadOwnership() final override;
-	UE_API virtual void RHIFlushResources() final override;
 	UE_API virtual FViewportRHIRef RHICreateViewport(void* WindowHandle, uint32 SizeX, uint32 SizeY, bool bIsFullscreen, EPixelFormat PixelFormat) final override;
 	UE_API virtual void RHIResizeViewport(FRHIViewport* Viewport, uint32 SizeX, uint32 SizeY, bool bIsFullscreen, EPixelFormat PixelFormat) final override;
 	UE_API virtual void RHITick(float DeltaTime) final override;
diff --git a/Engine/Source/Runtime/OpenGLDrv/Private/OpenGLDrvPrivate.h b/Engine/Source/Runtime/OpenGLDrv/Private/OpenGLDrvPrivate.h
--- a/Engine/Source/Runtime/OpenGLDrv/Private/OpenGLDrvPrivate.h
+++ b/Engine/Source/Runtime/OpenGLDrv/Private/OpenGLDrvPrivate.h
@@ -188,9 +188,6 @@ void PlatformDestroyOpenGLSurfaceContext(FPlatformOpenGLDevice* Device, FPlatfor
  */
 void PlatformRenderingContextSetup(FPlatformOpenGLDevice* Device);
 
-// Some platforms require flushing and rebinding renderbuffers after context/thread changes
-void PlatformFlushIfNeeded();
-
 /**
  * This is used to detach an OpenGL context from current thread.
  */
diff --git a/Engine/Source/Runtime/OpenGLDrv/Private/Windows/WindowsOpenGLPlatform.cpp b/Engine/Source/Runtime/OpenGLDrv/Private/Windows/WindowsOpenGLPlatform.cpp
--- a/Engine/Source/Runtime/OpenGLDrv/Private/Windows/WindowsOpenGLPlatform.cpp
+++ b/Engine/Source/Runtime/OpenGLDrv/Private/Windows/WindowsOpenGLPlatform.cpp
@@ -530,11 +530,6 @@ bool PlatformBlitToViewport(IRHICommandContext& RHICmdContext, FPlatformOpenGLDe
 	return true;
 }
 
-void PlatformFlushIfNeeded()
-{
-	glFinish();
-}
-
 void PlatformRenderingContextSetup(FPlatformOpenGLDevice* Device)
 {
 	check(Device && Device->RenderingContext.DeviceContext && Device->RenderingContext.OpenGLContext);
diff --git a/Engine/Source/Runtime/RHI/Public/DynamicRHI.h b/Engine/Source/Runtime/RHI/Public/DynamicRHI.h
--- a/Engine/Source/Runtime/RHI/Public/DynamicRHI.h
+++ b/Engine/Source/Runtime/RHI/Public/DynamicRHI.h
@@ -675,10 +675,6 @@ public:
 	virtual void RHIAcquireThreadOwnership() {}
 	virtual void RHIReleaseThreadOwnership() {}
 
-	// Flush driver resources. Typically called when switching contexts/threads
-	// FlushType: Flush RHI Thread
-	virtual void RHIFlushResources() = 0;
-
 	//  must be called from the main thread.
 	// FlushType: Thread safe
 	virtual FViewportRHIRef RHICreateViewport(void* WindowHandle, uint32 SizeX, uint32 SizeY, bool bIsFullscreen, EPixelFormat PixelFormat) = 0;
diff --git a/Engine/Source/Runtime/RHI/Public/RHICommandList.h b/Engine/Source/Runtime/RHI/Public/RHICommandList.h
--- a/Engine/Source/Runtime/RHI/Public/RHICommandList.h
+++ b/Engine/Source/Runtime/RHI/Public/RHICommandList.h
@@ -4665,33 +4665,27 @@ public:
 		ImmediateFlush(EImmediateFlushType::FlushRHIThread);  
 		GDynamicRHI->RHIRead3DSurfaceFloatData(Texture,Rect,ZMinMax,OutData,Flags);
 	}
-	
-	inline void FlushResources()
-	{
-		QUICK_SCOPE_CYCLE_COUNTER(STAT_RHIMETHOD_FlushResources_Flush);
-		ImmediateFlush(EImmediateFlushType::FlushRHIThread); 
-		 
-		return GDynamicRHI->RHIFlushResources();
-	}
 
-	inline void SubmitAndBlockUntilGPUIdle()
+	inline void SubmitAndBlockUntilGPUIdle(const bool bDeleteResources = false)
 	{
 		QUICK_SCOPE_CYCLE_COUNTER(STAT_RHIMETHOD_SubmitAndBlockUntilGPUIdle_Flush);
 
+		const EImmediateFlushType::Type FlushType = bDeleteResources ? EImmediateFlushType::FlushRHIThreadFlushResources : EImmediateFlushType::FlushRHIThread;
+		
 		// Ensure all prior work is submitted down to the GPU, and the RHI thread is idle.
-		ImmediateFlush(EImmediateFlushType::FlushRHIThread);
+		ImmediateFlush(FlushType);
 
 		// Block the calling thread (the render thread) until the GPU completes all work.
 		GDynamicRHI->RHIBlockUntilGPUIdle();
 	}
 	
-	//UE_DEPRECATED(5.3, "BlockUntilGPUIdle() is deprecated. Call SubmitAndBlockUntilGPUIdle() instead.")
+	UE_DEPRECATED(5.9, "BlockUntilGPUIdle() is deprecated. Call SubmitAndBlockUntilGPUIdle() instead.")
 	inline void BlockUntilGPUIdle()
 	{
 		this->SubmitAndBlockUntilGPUIdle();
 	}
 
-	//UE_DEPRECATED(5.3, "SubmitCommandsAndFlushGPU() is deprecated. Call SubmitAndBlockUntilGPUIdle() instead.")
+	UE_DEPRECATED(5.9, "SubmitCommandsAndFlushGPU() is deprecated. Call SubmitAndBlockUntilGPUIdle() instead.")
 	inline void SubmitCommandsAndFlushGPU()
 	{
 		this->SubmitAndBlockUntilGPUIdle();
@@ -5428,9 +5422,10 @@ inline void RHIUpdateTexture3D(FRHITexture* Texture, uint32 MipIndex, const stru
 	 FRHICommandListExecutor::GetImmediateCommandList().UpdateTexture3D(Texture, MipIndex, UpdateRegion, SourceRowPitch, SourceDepthPitch, SourceData);
 }
 
+UE_DEPRECATED(5.9, "RHIFlushResources has been deprecated, use ImmediateFlush or SubmitAndBlockUntilGPUIdle(true) instead to block")
 inline void RHIFlushResources()
 {
-	return FRHICommandListExecutor::GetImmediateCommandList().FlushResources();
+	FRHICommandListExecutor::GetImmediateCommandList().ImmediateFlush(EImmediateFlushType::FlushRHIThreadFlushResources);
 }
 
 UE_DEPRECATED(5.8, "RHIVirtualTextureSetFirstMipInMemory with an implied immediate command list is deprecated")
diff --git a/Engine/Source/Runtime/RHI/Public/RHIValidation.h b/Engine/Source/Runtime/RHI/Public/RHIValidation.h
--- a/Engine/Source/Runtime/RHI/Public/RHIValidation.h
+++ b/Engine/Source/Runtime/RHI/Public/RHIValidation.h
@@ -736,13 +736,6 @@ public:
 		RHI->RHIReleaseThreadOwnership();
 	}
 
-	// Flush driver resources. Typically called when switching contexts/threads
-	// FlushType: Flush RHI Thread
-	virtual void RHIFlushResources() override final
-	{
-		RHI->RHIFlushResources();
-	}
-
 	virtual FViewportRHIRef RHICreateViewport(void* WindowHandle, uint32 SizeX, uint32 SizeY, bool bIsFullscreen, EPixelFormat PixelFormat) override final
 	{
 		check(PixelFormat != PF_Unknown);
diff --git a/Engine/Source/Runtime/RenderCore/Private/DumpGPU.cpp b/Engine/Source/Runtime/RenderCore/Private/DumpGPU.cpp
--- a/Engine/Source/Runtime/RenderCore/Private/DumpGPU.cpp
+++ b/Engine/Source/Runtime/RenderCore/Private/DumpGPU.cpp
@@ -549,10 +549,8 @@ public:
 
 		// Flush the RHI resource memory so the readback memory can be fully reused in the next resource dump.
 		{
-			RHICmdList.SubmitCommandsAndFlushGPU();
-			RHICmdList.BlockUntilGPUIdle();
-			RHICmdList.FlushResources();
-			RHICmdList.ImmediateFlush(EImmediateFlushType::FlushRHIThreadFlushResources);
+			constexpr bool bFlushResources = true;
+			RHICmdList.SubmitAndBlockUntilGPUIdle(bFlushResources);
 		}
 	}
 
@@ -1288,8 +1286,7 @@ public:
 
 				Fence->Clear();
 				RHICmdList.WriteGPUFence(Fence);
-				RHICmdList.SubmitCommandsAndFlushGPU();
-				RHICmdList.BlockUntilGPUIdle();
+				RHICmdList.SubmitAndBlockUntilGPUIdle();
 			}
 
 			void* Content = nullptr;
@@ -1782,8 +1779,7 @@ public:
 							FTimeBucketMeasure TimeBucketMeasure(this, ETimingBucket::GPUWait);
 							Fence->Clear();
 							RHICmdList.WriteGPUFence(Fence);
-							RHICmdList.SubmitCommandsAndFlushGPU();
-							RHICmdList.BlockUntilGPUIdle();
+							RHICmdList.SubmitAndBlockUntilGPUIdle();
 						}
 
 						void* Content = RHICmdList.LockStagingBuffer(StagingBuffer, Fence.GetReference(), 0, CopyByteSize);
@@ -1942,8 +1938,7 @@ public:
 		check(IsInRenderingThread());
 		check(bStream);
 
-		RHICmdList.SubmitCommandsAndFlushGPU();
-		RHICmdList.BlockUntilGPUIdle();
+		RHICmdList.SubmitAndBlockUntilGPUIdle();
 
 		LandCompletedResources(RHICmdList);
 
@@ -1973,8 +1968,7 @@ public:
 			check(Entry->Status == FStagingResourceStatus::Unused);
 		}
 
-		RHICmdList.SubmitCommandsAndFlushGPU();
-		RHICmdList.BlockUntilGPUIdle();
+		RHICmdList.SubmitAndBlockUntilGPUIdle();
 
 		// Releases all staging resources
 		StagingTexturePool.Reset();
@@ -2381,7 +2375,7 @@ void FRDGResourceDumpContext::Start()
 		check(IsInRenderingThread());
 		GRDGResourceDumpContext_RenderThread = this;
 
-		ImmediateRHICmdList.SubmitCommandsAndFlushGPU();
+		ImmediateRHICmdList.SubmitAndBlockUntilGPUIdle();
 
 		// Disable the validation for BUF_SourceCopy so that all buffers can be copied into staging buffer for CPU readback.
 		#if ENABLE_RHI_VALIDATION
@@ -2407,7 +2401,7 @@ void FRDGResourceDumpContext::Finish()
 				WaitAndReleaseStagingResources(ImmediateRHICmdList);
 			}
 
-			ImmediateRHICmdList.SubmitCommandsAndFlushGPU();
+			ImmediateRHICmdList.SubmitAndBlockUntilGPUIdle();
 			#if ENABLE_RHI_VALIDATION
 				GRHIValidateBufferSourceCopy = true;
 			#endif
diff --git a/Engine/Source/Runtime/RenderCore/Private/RenderGraphBuilder.cpp b/Engine/Source/Runtime/RenderCore/Private/RenderGraphBuilder.cpp
--- a/Engine/Source/Runtime/RenderCore/Private/RenderGraphBuilder.cpp
+++ b/Engine/Source/Runtime/RenderCore/Private/RenderGraphBuilder.cpp
@@ -2151,8 +2151,7 @@ void FRDGBuilder::Execute()
 
 				if (GRDGDebugFlushGPU && !bSupportsAsyncCompute)
 				{
-					RHICmdList.SubmitCommandsAndFlushGPU();
-					RHICmdList.BlockUntilGPUIdle();
+					RHICmdList.SubmitAndBlockUntilGPUIdle();
 				}
 			}
 		}
diff --git a/Engine/Source/Runtime/Renderer/Private/ReflectionEnvironmentCapture.cpp b/Engine/Source/Runtime/Renderer/Private/ReflectionEnvironmentCapture.cpp
--- a/Engine/Source/Runtime/Renderer/Private/ReflectionEnvironmentCapture.cpp
+++ b/Engine/Source/Runtime/Renderer/Private/ReflectionEnvironmentCapture.cpp
@@ -635,7 +635,7 @@ void CaptureSceneToScratchCubemap(
 
 		// some platforms may not be able to keep enqueueing commands like crazy, this will
 		// allow them to restart their command buffers
-		InRHICmdList.SubmitCommandsAndFlushGPU();
+		InRHICmdList.SubmitAndBlockUntilGPUIdle();
 	});
 
 	const FViewInfo& View = SceneRenderer->Views[0];
diff --git a/Engine/Source/Runtime/Renderer/Private/SceneViewState.cpp b/Engine/Source/Runtime/Renderer/Private/SceneViewState.cpp
--- a/Engine/Source/Runtime/Renderer/Private/SceneViewState.cpp
+++ b/Engine/Source/Runtime/Renderer/Private/SceneViewState.cpp
@@ -849,8 +849,8 @@ void FSceneViewState::SystemMemoryMirrorBackup(FSceneViewStateSystemMemoryMirror
 			SystemMemoryForEachTexture(RHICmdList, *SystemMemoryMirror, *this, SystemMemoryBackupTextureBegin);
 		}
 
-		RHICmdList.BlockUntilGPUIdle();
-		RHICmdList.FlushResources();
+		constexpr bool bFlushResources = true;
+		RHICmdList.SubmitAndBlockUntilGPUIdle(bFlushResources);
 
 		// Clear out any temporary textures used to copy depth
 		SystemMemoryMirror->TemporaryTextures.Reset();
diff --git a/Engine/Source/Runtime/Renderer/Private/Substrate/SubstrateRoughRefraction.cpp b/Engine/Source/Runtime/Renderer/Private/Substrate/SubstrateRoughRefraction.cpp
--- a/Engine/Source/Runtime/Renderer/Private/Substrate/SubstrateRoughRefraction.cpp
+++ b/Engine/Source/Runtime/Renderer/Private/Substrate/SubstrateRoughRefraction.cpp
@@ -546,8 +546,7 @@ void SubstrateRoughRefractionRnD(FRDGBuilder& GraphBuilder, const FViewInfo& Vie
 					FGPUFenceRHIRef Fence = RHICreateGPUFence(FenceName);
 					Fence->Clear();
 					RHICmdList.WriteGPUFence(Fence);
-					RHICmdList.SubmitCommandsAndFlushGPU();
-					RHICmdList.BlockUntilGPUIdle();
+					RHICmdList.SubmitAndBlockUntilGPUIdle();
 
 					// Lock the buffer for read back
 					void* RoughRefracData = RHICmdList.LockStagingBuffer(StagingBuffer, Fence.GetReference(), 0, RoughRefracDataByteSize);
@@ -600,4 +599,4 @@ void SubstrateRoughRefractionRnD(FRDGBuilder& GraphBuilder, const FViewInfo& Vie
 }
 
 
-} // namespace Substrate
\ No newline at end of file
+} // namespace Substrate
diff --git a/Engine/Source/Runtime/Renderer/Private/Tests/RayTracingTestbed.cpp b/Engine/Source/Runtime/Renderer/Private/Tests/RayTracingTestbed.cpp
--- a/Engine/Source/Runtime/Renderer/Private/Tests/RayTracingTestbed.cpp
+++ b/Engine/Source/Runtime/Renderer/Private/Tests/RayTracingTestbed.cpp
@@ -271,7 +271,7 @@ bool RunRayTracingTestbed_RenderThread(const FString& Parameters)
 
 	if (bValidateResults)
 	{
-		RHICmdList.BlockUntilGPUIdle();
+		RHICmdList.SubmitAndBlockUntilGPUIdle();
 
 		// Read back and validate occlusion trace results
 
diff --git a/Engine/Source/Runtime/SlateRHIRenderer/Private/SlateRHIRenderer.cpp b/Engine/Source/Runtime/SlateRHIRenderer/Private/SlateRHIRenderer.cpp
--- a/Engine/Source/Runtime/SlateRHIRenderer/Private/SlateRHIRenderer.cpp
+++ b/Engine/Source/Runtime/SlateRHIRenderer/Private/SlateRHIRenderer.cpp
@@ -495,7 +495,7 @@ void FSlateViewportInfo::SetCustomPresent(FRHICustomPresent* CustomPresent)
 void FSlateViewportInfo::ReleaseRHI()
 {
 	// Full GPU sync here to simplify memory lifetime of the underlying resource.
-	FRHICommandListImmediate::Get().BlockUntilGPUIdle();
+	FRHICommandListImmediate::Get().SubmitAndBlockUntilGPUIdle();
 
 	if (IsViewportRHI())
 	{
diff --git a/Engine/Source/Runtime/VulkanRHI/Private/VulkanDynamicRHI.h b/Engine/Source/Runtime/VulkanRHI/Private/VulkanDynamicRHI.h
--- a/Engine/Source/Runtime/VulkanRHI/Private/VulkanDynamicRHI.h
+++ b/Engine/Source/Runtime/VulkanRHI/Private/VulkanDynamicRHI.h
@@ -212,7 +212,6 @@ public:
 	virtual FTextureRHIRef RHIGetViewportBackBuffer(FRHIViewport* Viewport) final override;
 	virtual void RHIAliasTextureResources(FTextureRHIRef& DestTexture, FTextureRHIRef& SrcTexture) final override;
 	virtual FTextureRHIRef RHICreateAliasedTexture(FTextureRHIRef& SourceTexture) final override;
-	virtual void RHIFlushResources() final override;
 	virtual FViewportRHIRef RHICreateViewport(void* WindowHandle, uint32 SizeX, uint32 SizeY, bool bIsFullscreen, EPixelFormat PixelFormat) final override;
 	virtual void RHIResizeViewport(FRHIViewport* Viewport, uint32 SizeX, uint32 SizeY, bool bIsFullscreen, EPixelFormat PixelFormat) final override;
 	virtual void RHIEndDrawingViewport(FRHICommandListImmediate& RHICmdList, FRHIViewport* Viewport, const FRHIPresentArgs& PresentArgs) final override;
diff --git a/Engine/Source/Runtime/VulkanRHI/Private/VulkanRHI.cpp b/Engine/Source/Runtime/VulkanRHI/Private/VulkanRHI.cpp
--- a/Engine/Source/Runtime/VulkanRHI/Private/VulkanRHI.cpp
+++ b/Engine/Source/Runtime/VulkanRHI/Private/VulkanRHI.cpp
@@ -1159,7 +1159,7 @@ void FVulkanDynamicRHI::RHIEndFrame_RenderThread(FRHICommandListImmediate& RHICm
 
 		check(Context.IsImmediate());
 
-		bool bTrimMemory = false;
+		bool bTrimMemory = true;
 		Context.FreeUnusedCmdBuffers(bTrimMemory);
 
 		Context.Device.GetStagingManager().ProcessPendingFree(false, true);
@@ -1403,12 +1403,6 @@ bool FVulkanDynamicRHI::RHIGetAvailableResolutions(FScreenResolutionArray& Resol
 	return false;
 }
 
-void FVulkanDynamicRHI::RHIFlushResources()
-{
-	bool bTrimMemory = true;
-	GetDevice()->GetImmediateContext().FreeUnusedCmdBuffers(bTrimMemory);
-}
-
 // IVulkanDynamicRHI interface
 
 uint32 FVulkanDynamicRHI::RHIGetVulkanVersion() const
diff --git a/Engine/Source/Runtime/VulkanRHI/Private/VulkanRenderTarget.cpp b/Engine/Source/Runtime/VulkanRHI/Private/VulkanRenderTarget.cpp
--- a/Engine/Source/Runtime/VulkanRHI/Private/VulkanRenderTarget.cpp
+++ b/Engine/Source/Runtime/VulkanRHI/Private/VulkanRenderTarget.cpp
@@ -315,7 +315,7 @@ void FVulkanDynamicRHI::RHIMapStagingSurface(FRHITexture* TextureRHI, FRHIGPUFen
 		FRHICommandListImmediate& RHICmdList = FRHICommandListImmediate::Get();
 		RHICmdList.ImmediateFlush(EImmediateFlushType::FlushRHIThread);
 
-		// SubmitCommandsAndFlushGPU might update fence state if it was tied to a previously submitted command buffer.
+		// SubmitAndBlockUntilGPUIdle might update fence state if it was tied to a previously submitted command buffer.
 		// Its state will have been updated from Submitted to NeedReset, and would assert in WaitForCmdBuffer (which is not needed in such a case)
 		FenceRHI->Wait(RHICmdList, FRHIGPUMask::All());
 	}
diff --git a/Engine/Source/Runtime/Windows/D3D11RHI/Private/D3D11Device.cpp b/Engine/Source/Runtime/Windows/D3D11RHI/Private/D3D11Device.cpp
--- a/Engine/Source/Runtime/Windows/D3D11RHI/Private/D3D11Device.cpp
+++ b/Engine/Source/Runtime/Windows/D3D11RHI/Private/D3D11Device.cpp
@@ -755,11 +755,6 @@ void FD3D11DynamicRHI::CleanupD3DDevice()
 	}
 }
 
-void FD3D11DynamicRHI::RHIFlushResources()
-{
-	// Nothing to do (yet!)
-}
-
 void* FD3D11DynamicRHI::RHIGetNativeDevice()
 {
 	return (void*)Direct3DDevice.GetReference();
diff --git a/Engine/Source/Runtime/Windows/D3D11RHI/Private/D3D11RHIPrivate.h b/Engine/Source/Runtime/Windows/D3D11RHI/Private/D3D11RHIPrivate.h
--- a/Engine/Source/Runtime/Windows/D3D11RHI/Private/D3D11RHIPrivate.h
+++ b/Engine/Source/Runtime/Windows/D3D11RHI/Private/D3D11RHIPrivate.h
@@ -291,7 +291,6 @@ public:
 	D3D11RHI_API virtual FTextureRHIRef RHIGetViewportBackBuffer(FRHIViewport* Viewport) final override;
 	D3D11RHI_API virtual void RHIAliasTextureResources(FTextureRHIRef& DestTexture, FTextureRHIRef& SrcTexture) final override;
 	D3D11RHI_API virtual FTextureRHIRef RHICreateAliasedTexture(FTextureRHIRef& SourceTexture) final override;
-	D3D11RHI_API virtual void RHIFlushResources() final override;
 	D3D11RHI_API virtual FViewportRHIRef RHICreateViewport(void* WindowHandle, uint32 SizeX, uint32 SizeY, bool bIsFullscreen, EPixelFormat PixelFormat) override;
 	D3D11RHI_API virtual void RHIResizeViewport(FRHIViewport* Viewport, uint32 SizeX, uint32 SizeY, bool bIsFullscreen, EPixelFormat PixelFormat) final override;
 	D3D11RHI_API virtual void RHIHandleDisplayChange() final override;
